Added initial draft of Universal Object Locator mechanism

Universal Object Locators (UOLs) are designed to support referencing objects that don't fit neatly into a basic outer->inner path representation. Examples might include transient actors, dynamically created objects, or objects that need to be referenced by an external ID or using external lookup logic. Specifically this might be an object spawned by Sequencer, a transient object on a USD stage, or a gameplay-specific object created by a game system.

A UOL comprises zero or more 'fragments': atomic pieces of data and logic that defines how to lookup or load an object based on a context. Fragment types are globally registered as part of module initialization.
UOLs are hashable, and support string conversion that conforms to RFC3986 so they can be used as URIs (though that is not a current use-case). In order to support this type of string conversion, the 'path' part of of a UOL defines the fragment types, and the query string is used to encode the payload data for each fragment. This allows us to support a more diverse set of characters as part of payload strings (ie, / : and .) which are otherwise unsupported as part of the path.

An example UOL to an anim instance might look like: uobj://actor/subobj/animinst?payload0=/Path/To/Package.LevelName:PathToActor&payload1=ComponentName

diff --git a/Engine/Source/Editor/UniversalObjectLocatorEditor/Private/ActorLocatorEditor.cpp b/Engine/Source/Editor/UniversalObjectLocatorEditor/Private/ActorLocatorEditor.cpp
new file mode 100644
--- /dev/null
+++ b/Engine/Source/Editor/UniversalObjectLocatorEditor/Private/ActorLocatorEditor.cpp
@@ -0,0 +1,136 @@
+// Copyright Epic Games, Inc. All Rights Reserved.
+
+#include "ActorLocatorEditor.h"
+#include "Modules/ModuleManager.h"
+#include "UniversalObjectLocatorFragmentTypeHandle.h"
+#include "UniversalObjectLocator.h"
+#include "UniversalObjectLocatorEditor.h"
+#include "IUniversalObjectLocatorEditorModule.h"
+#include "IUniversalObjectLocatorCustomization.h"
+
+#include "SceneOutlinerDragDrop.h"
+#include "DragAndDrop/ActorDragDropOp.h"
+
+#include "PropertyCustomizationHelpers.h"
+
+#include "GameFramework/Actor.h"
+
+
+#define LOCTEXT_NAMESPACE "ActorLocatorEditor"
+
+namespace UE::UniversalObjectLocator
+{
+
+bool FActorLocatorEditor::IsDragSupported(TSharedPtr<FDragDropOperation> DragOperation, UObject* Context) const
+{
+	TSharedPtr<FActorDragDropOp> ActorDrag;
+
+	if (DragOperation->IsOfType<FSceneOutlinerDragDropOp>())
+	{
+		FSceneOutlinerDragDropOp* SceneOutlinerOp = static_cast<FSceneOutlinerDragDropOp*>(DragOperation.Get());
+		ActorDrag = SceneOutlinerOp->GetSubOp<FActorDragDropOp>();
+	}
+	else if (DragOperation->IsOfType<FActorDragDropOp>())
+	{
+		ActorDrag = StaticCastSharedPtr<FActorDragDropOp>(DragOperation);
+	}		
+
+	if (ActorDrag)
+	{
+		for (const TWeakObjectPtr<AActor>& WeakActor : ActorDrag->Actors)
+		{
+			if (WeakActor.Get())
+			{
+				return true;
+			}
+		}
+	}
+
+	return false;
+}
+
+UObject* FActorLocatorEditor::ResolveDragOperation(TSharedPtr<FDragDropOperation> DragOperation, UObject* Context) const
+{
+	TSharedPtr<FActorDragDropOp> ActorDrag;
+
+	if (DragOperation->IsOfType<FSceneOutlinerDragDropOp>())
+	{
+		FSceneOutlinerDragDropOp* SceneOutlinerOp = static_cast<FSceneOutlinerDragDropOp*>(DragOperation.Get());
+		ActorDrag = SceneOutlinerOp->GetSubOp<FActorDragDropOp>();
+	}
+	else if (DragOperation->IsOfType<FActorDragDropOp>())
+	{
+		ActorDrag = StaticCastSharedPtr<FActorDragDropOp>(DragOperation);
+	}
+
+	if (ActorDrag)
+	{
+		for (const TWeakObjectPtr<AActor>& WeakActor : ActorDrag->Actors)
+		{
+			if (AActor* Actor = WeakActor.Get())
+			{
+				return Actor;
+			}
+		}
+	}
+
+	return nullptr;
+}
+
+TSharedPtr<SWidget> FActorLocatorEditor::MakeEditUI(TSharedPtr<IUniversalObjectLocatorCustomization> Customization)
+{
+	TSharedRef<SObjectPropertyEntryBox> EditWidget = SNew(SObjectPropertyEntryBox)
+	.ObjectPath(Customization.ToSharedRef(), &IUniversalObjectLocatorCustomization::GetPathToObject)
+	// Disabling this for now since it causes issues
+	//.PropertyHandle(Customization->GetProperty())
+	.AllowedClass(AActor::StaticClass())
+	.OnObjectChanged(this, &FActorLocatorEditor::OnSetObject, TWeakPtr<IUniversalObjectLocatorCustomization>(Customization))
+	.AllowClear(true)
+	.DisplayUseSelected(true)
+	.DisplayBrowse(true)
+	.DisplayThumbnail(true);
+
+	float MinWidth = 100.f;
+	float MaxWidth = 500.f;
+	EditWidget->GetDesiredWidth(MinWidth, MaxWidth);
+
+	return SNew(SBox)
+	.MinDesiredWidth(MinWidth)
+	.MaxDesiredWidth(MaxWidth)
+	[
+		EditWidget
+	];
+}
+
+FText FActorLocatorEditor::GetDisplayText() const
+{
+	return LOCTEXT("ExternalActorLocatorName", "Actor");
+}
+
+FText FActorLocatorEditor::GetDisplayTooltip() const
+{
+	return LOCTEXT("ExternalActorLocatorTooltip", "Change this to an actor reference");
+}
+
+FSlateIcon FActorLocatorEditor::GetDisplayIcon() const
+{
+	return FSlateIcon();
+}
+
+void FActorLocatorEditor::OnSetObject(const FAssetData& InNewObject, TWeakPtr<IUniversalObjectLocatorCustomization> WeakCustomization)
+{
+	TSharedPtr<IUniversalObjectLocatorCustomization> Customization = WeakCustomization.Pin();
+	if (!Customization)
+	{
+		return;
+	}
+
+	UObject* Object = InNewObject.FastGetAsset(true);
+
+	FUniversalObjectLocator NewRef(Object);
+	Customization->SetValue(MoveTemp(NewRef));
+}
+
+} // namespace UE::UniversalObjectLocator
+
+#undef LOCTEXT_NAMESPACE
diff --git a/Engine/Source/Editor/UniversalObjectLocatorEditor/Private/ActorLocatorEditor.h b/Engine/Source/Editor/UniversalObjectLocatorEditor/Private/ActorLocatorEditor.h
new file mode 100644
--- /dev/null
+++ b/Engine/Source/Editor/UniversalObjectLocatorEditor/Private/ActorLocatorEditor.h
@@ -0,0 +1,36 @@
+// Copyright Epic Games, Inc. All Rights Reserved.
+
+#pragma once
+
+#include "CoreTypes.h"
+#include "UniversalObjectLocatorEditor.h"
+#include "Templates/SharedPointerFwd.h"
+
+struct FAssetData;
+
+namespace UE::UniversalObjectLocator
+{
+
+class IUniversalObjectLocatorCustomization;
+
+class FActorLocatorEditor : public ILocatorEditor
+{
+	bool IsDragSupported(TSharedPtr<FDragDropOperation> DragOperation, UObject* Context) const override;
+
+	UObject* ResolveDragOperation(TSharedPtr<FDragDropOperation> DragOperation, UObject* Context) const override;
+
+	TSharedPtr<SWidget> MakeEditUI(TSharedPtr<IUniversalObjectLocatorCustomization> Customization) override;
+
+	FText GetDisplayText() const override;
+
+	FText GetDisplayTooltip() const override;
+
+	FSlateIcon GetDisplayIcon() const override;
+
+private:
+
+	void OnSetObject(const FAssetData& InNewObject, TWeakPtr<IUniversalObjectLocatorCustomization> WeakCustomization);
+};
+
+} // namespace UE::UniversalObjectLocator
+
diff --git a/Engine/Source/Editor/UniversalObjectLocatorEditor/Private/AnimInstanceLocatorEditor.cpp b/Engine/Source/Editor/UniversalObjectLocatorEditor/Private/AnimInstanceLocatorEditor.cpp
new file mode 100644
--- /dev/null
+++ b/Engine/Source/Editor/UniversalObjectLocatorEditor/Private/AnimInstanceLocatorEditor.cpp
@@ -0,0 +1,298 @@
+// Copyright Epic Games, Inc. All Rights Reserved.
+
+#include "AnimInstanceLocatorEditor.h"
+#include "UniversalObjectLocatorFragmentTypeHandle.h"
+#include "UniversalObjectLocators/AnimInstanceLocatorFragment.h"
+#include "UniversalObjectLocator.h"
+#include "UniversalObjectLocatorEditor.h"
+#include "IUniversalObjectLocatorEditorModule.h"
+#include "IUniversalObjectLocatorCustomization.h"
+#include "ISequencerModule.h"
+#include "SceneOutlinerDragDrop.h"
+#include "DragAndDrop/ActorDragDropOp.h"
+#include "Widgets/Input/SComboButton.h"
+#include "Framework/MultiBox/MultiBoxBuilder.h"
+
+#include "PropertyCustomizationHelpers.h"
+
+#include "GameFramework/Actor.h"
+#include "Components/SkeletalMeshComponent.h"
+
+
+#define LOCTEXT_NAMESPACE "AnimInstanceLocatorEditor"
+
+namespace UE::UniversalObjectLocator
+{
+
+class SAnimInstanceLocatorEditorUI : public SCompoundWidget
+{
+	SLATE_BEGIN_ARGS(SAnimInstanceLocatorEditorUI){}
+	SLATE_END_ARGS()
+
+	void Construct(const FArguments& InArgs, TSharedPtr<IUniversalObjectLocatorCustomization> InCustomization)
+	{
+		WeakCustomization = InCustomization;
+
+		TSharedRef<SObjectPropertyEntryBox> EditWidget = SNew(SObjectPropertyEntryBox)
+		.ObjectPath(InCustomization.ToSharedRef(), &IUniversalObjectLocatorCustomization::GetPathToObject)
+		.AllowedClass(USkeletalMeshComponent::StaticClass())
+		.OnObjectChanged(this, &SAnimInstanceLocatorEditorUI::OnSetObject)
+		.AllowClear(true)
+		.DisplayUseSelected(true)
+		.DisplayBrowse(true)
+		.DisplayThumbnail(true);
+
+		float MinWidth = 100.f;
+		float MaxWidth = 500.f;
+		EditWidget->GetDesiredWidth(MinWidth, MaxWidth);
+
+		ChildSlot
+		[
+			SNew(SHorizontalBox)
+
+			+SHorizontalBox::Slot()
+			[
+				SNew(SBox)
+				.MinDesiredWidth(MinWidth)
+				.MaxDesiredWidth(MaxWidth)
+				[
+					EditWidget
+				]
+			]
+
+			+ SHorizontalBox::Slot()
+			.AutoWidth()
+			[
+				SNew(SComboButton)
+				.OnGetMenuContent(this, &SAnimInstanceLocatorEditorUI::CreateAnimInstanceTypeMenuContent)
+				.ButtonContent()
+				[
+					SNew(STextBlock)
+					.Text(this, &SAnimInstanceLocatorEditorUI::GetCurrentAnimInstanceTypeText)
+				]
+			]
+		];
+	}
+
+private:
+
+	void OnSetObject(const FAssetData& InNewObject)
+	{
+		TSharedPtr<IUniversalObjectLocatorCustomization> Customization = WeakCustomization.Pin();
+		if (!Customization)
+		{
+			return;
+		}
+
+		UObject* Object  = InNewObject.FastGetAsset(true);
+
+		if (Object)
+		{
+			FUniversalObjectLocator NewRef(Object);
+			Customization->SetValue(MoveTemp(NewRef));
+		}
+	}
+
+	TSharedRef<SWidget> CreateAnimInstanceTypeMenuContent()
+	{
+		const bool bCloseAfterSelection = true;
+		FMenuBuilder MenuBuilder(bCloseAfterSelection, nullptr);
+
+		FCanExecuteAction AlwaysExecute = FCanExecuteAction::CreateLambda([]{ return true; });
+
+		MenuBuilder.AddMenuEntry(
+			LOCTEXT("Menu_AnimInstanceLabel", "Anim Instance"),
+			LOCTEXT("Menu_AnimInstanceTooltip", "Bind to the Anim Instance on the selected component"),
+			FSlateIcon(),
+			FUIAction(
+				FExecuteAction::CreateSP(this, &SAnimInstanceLocatorEditorUI::ChangeType, EAnimInstanceLocatorFragmentType::AnimInstance),
+				AlwaysExecute,
+				FIsActionChecked::CreateSP(this, &SAnimInstanceLocatorEditorUI::CompareCurrentType, EAnimInstanceLocatorFragmentType::AnimInstance)
+			),
+			NAME_None,
+			EUserInterfaceActionType::RadioButton
+		);
+		MenuBuilder.AddMenuEntry(
+			LOCTEXT("Menu_PPAnimInstanceLabel", "Post Process Anim Instance"),
+			LOCTEXT("Menu_PPAnimInstanceTooltip", "Bind to the Post Process Anim Instance on the selected component"),
+			FSlateIcon(),
+			FUIAction(
+				FExecuteAction::CreateSP(this, &SAnimInstanceLocatorEditorUI::ChangeType, EAnimInstanceLocatorFragmentType::PostProcessAnimInstance),
+				AlwaysExecute,
+				FIsActionChecked::CreateSP(this, &SAnimInstanceLocatorEditorUI::CompareCurrentType, EAnimInstanceLocatorFragmentType::PostProcessAnimInstance)
+			),
+			NAME_None,
+			EUserInterfaceActionType::RadioButton
+		);
+		return MenuBuilder.MakeWidget();
+	}
+
+	FText GetCurrentAnimInstanceTypeText() const
+	{
+		TOptional<EAnimInstanceLocatorFragmentType> CommonType = GetCurrentType();
+		if (!CommonType)
+		{
+			return LOCTEXT("AnimInstanceMixedLabel", "None");
+		}
+
+		if (CommonType.GetValue() == EAnimInstanceLocatorFragmentType::AnimInstance)
+		{
+			return LOCTEXT("AnimInstanceLabel", "Anim Instance");
+		}
+		return LOCTEXT("PostProcessAnimInstanceLabel", "Post Process Anim Instance");
+	}
+
+	void ChangeType(EAnimInstanceLocatorFragmentType InType)
+	{
+		TSharedPtr<IUniversalObjectLocatorCustomization> Customization = WeakCustomization.Pin();
+		if (Customization)
+		{
+			TSharedPtr<IPropertyHandle>    Property          = Customization->GetProperty();
+			const FStructProperty*         RawProperty       = CastFieldChecked<const FStructProperty>(Property->GetProperty());
+
+			Property->NotifyPreChange();
+			Property->EnumerateRawData([this, InType](void* RawData, const int32, const int32){
+
+				FUniversalObjectLocatorFragment* Locator = static_cast<FUniversalObjectLocator*>(RawData)->GetLastFragment();
+
+				FAnimInstanceLocatorFragment* Payload = nullptr;
+				if (Locator && Locator->TryGetPayloadAs(FAnimInstanceLocatorFragment::FragmentType, Payload))
+				{
+					Payload->Type = InType;
+				}
+
+				// Continue to the next property value
+				return true;
+			});
+
+			Property->NotifyPostChange(EPropertyChangeType::ValueSet);
+			Property->NotifyFinishedChangingProperties();
+		}
+	}
+
+	bool CompareCurrentType(EAnimInstanceLocatorFragmentType InType) const
+	{
+		TOptional<EAnimInstanceLocatorFragmentType> Type = GetCurrentType();
+		return Type.IsSet() && InType == Type.GetValue();
+	}
+
+	TOptional<EAnimInstanceLocatorFragmentType> GetCurrentType() const
+	{
+		TOptional<EAnimInstanceLocatorFragmentType> CommonType;
+
+		TSharedPtr<IUniversalObjectLocatorCustomization> Customization = WeakCustomization.Pin();
+		if (Customization)
+		{
+			TSharedPtr<IPropertyHandle> Property    = Customization->GetProperty();
+			const FStructProperty*      RawProperty = CastFieldChecked<const FStructProperty>(Property->GetProperty());
+
+			Property->EnumerateRawData([this, &CommonType](void* RawData, const int32, const int32){
+
+				FUniversalObjectLocatorFragment* LastFragment = static_cast<FUniversalObjectLocator*>(RawData)->GetLastFragment();
+
+				FAnimInstanceLocatorFragment* Payload = nullptr;
+				if (LastFragment && LastFragment->TryGetPayloadAs(FAnimInstanceLocatorFragment::FragmentType, Payload))
+				{
+					if (!CommonType.IsSet())
+					{
+						CommonType = Payload->Type;
+					}
+					else if (CommonType.GetValue() != Payload->Type)
+					{
+						// Different type - common type is undefined
+						CommonType.Reset();
+						// Stop iterating and return an unset optional
+						return false;
+					}
+				}
+
+				// Continue to the next property value
+				return true;
+			});
+		}
+
+		return CommonType;
+	}
+
+	TWeakPtr<IUniversalObjectLocatorCustomization> WeakCustomization;
+	EAnimInstanceLocatorFragmentType CurrentType;
+};
+
+bool FAnimInstanceLocatorEditor::IsDragSupported(TSharedPtr<FDragDropOperation> DragOperation, UObject* Context) const
+{
+	TSharedPtr<FActorDragDropOp> ActorDrag;
+
+	if (DragOperation->IsOfType<FSceneOutlinerDragDropOp>())
+	{
+		FSceneOutlinerDragDropOp* SceneOutlinerOp = static_cast<FSceneOutlinerDragDropOp*>(DragOperation.Get());
+		ActorDrag = SceneOutlinerOp->GetSubOp<FActorDragDropOp>();
+	}
+	else if (DragOperation->IsOfType<FActorDragDropOp>())
+	{
+		ActorDrag = StaticCastSharedPtr<FActorDragDropOp>(DragOperation);
+	}		
+
+	if (ActorDrag)
+	{
+		for (const TWeakObjectPtr<AActor>& WeakActor : ActorDrag->Actors)
+		{
+			if (WeakActor.Get())
+			{
+				return true;
+			}
+		}
+	}
+
+	return false;
+}
+
+UObject* FAnimInstanceLocatorEditor::ResolveDragOperation(TSharedPtr<FDragDropOperation> DragOperation, UObject* Context) const
+{
+	TSharedPtr<FActorDragDropOp> ActorDrag;
+
+	if (DragOperation->IsOfType<FSceneOutlinerDragDropOp>())
+	{
+		FSceneOutlinerDragDropOp* SceneOutlinerOp = static_cast<FSceneOutlinerDragDropOp*>(DragOperation.Get());
+		ActorDrag = SceneOutlinerOp->GetSubOp<FActorDragDropOp>();
+	}
+	else if (DragOperation->IsOfType<FActorDragDropOp>())
+	{
+		ActorDrag = StaticCastSharedPtr<FActorDragDropOp>(DragOperation);
+	}
+
+	if (ActorDrag)
+	{
+		for (const TWeakObjectPtr<AActor>& WeakActor : ActorDrag->Actors)
+		{
+			if (AActor* Actor = WeakActor.Get())
+			{
+				return Actor;
+			}
+		}
+	}
+
+	return nullptr;
+}
+
+TSharedPtr<SWidget> FAnimInstanceLocatorEditor::MakeEditUI(TSharedPtr<IUniversalObjectLocatorCustomization> Customization)
+{
+	return SNew(SAnimInstanceLocatorEditorUI, Customization);
+}
+
+FText FAnimInstanceLocatorEditor::GetDisplayText() const
+{
+	return LOCTEXT("AnimInstanceLocatorName", "Anim Instance");
+}
+FText FAnimInstanceLocatorEditor::GetDisplayTooltip() const
+{
+	return LOCTEXT("AnimInstanceLocatorTooltip", "Change this to a reference to an Anim Instance");
+}
+FSlateIcon FAnimInstanceLocatorEditor::GetDisplayIcon() const
+{
+	return FSlateIcon();
+}
+
+} // namespace UE::UniversalObjectLocator
+
+
+#undef LOCTEXT_NAMESPACE
diff --git a/Engine/Source/Editor/UniversalObjectLocatorEditor/Private/AnimInstanceLocatorEditor.h b/Engine/Source/Editor/UniversalObjectLocatorEditor/Private/AnimInstanceLocatorEditor.h
new file mode 100644
--- /dev/null
+++ b/Engine/Source/Editor/UniversalObjectLocatorEditor/Private/AnimInstanceLocatorEditor.h
@@ -0,0 +1,28 @@
+// Copyright Epic Games, Inc. All Rights Reserved.
+
+#pragma once
+
+#include "CoreTypes.h"
+#include "UniversalObjectLocatorEditor.h"
+
+namespace UE::UniversalObjectLocator
+{
+
+class FAnimInstanceLocatorEditor : public ILocatorEditor
+{
+public:
+	bool IsDragSupported(TSharedPtr<FDragDropOperation> DragOperation, UObject* Context) const override;
+
+	UObject* ResolveDragOperation(TSharedPtr<FDragDropOperation> DragOperation, UObject* Context) const override;
+
+	TSharedPtr<SWidget> MakeEditUI(TSharedPtr<IUniversalObjectLocatorCustomization> Customization) override;
+
+	FText GetDisplayText() const override;
+
+	FText GetDisplayTooltip() const override;
+
+	FSlateIcon GetDisplayIcon() const override;
+};
+
+
+} // namespace UE::UniversalObjectLocator
diff --git a/Engine/Source/Editor/UniversalObjectLocatorEditor/Private/AssetLocatorEditor.cpp b/Engine/Source/Editor/UniversalObjectLocatorEditor/Private/AssetLocatorEditor.cpp
new file mode 100644
--- /dev/null
+++ b/Engine/Source/Editor/UniversalObjectLocatorEditor/Private/AssetLocatorEditor.cpp
@@ -0,0 +1,109 @@
+// Copyright Epic Games, Inc. All Rights Reserved.
+
+#include "AssetLocatorEditor.h"
+#include "Modules/ModuleManager.h"
+#include "UniversalObjectLocatorFragmentTypeHandle.h"
+#include "UniversalObjectLocator.h"
+#include "UniversalObjectLocatorEditor.h"
+#include "IUniversalObjectLocatorEditorModule.h"
+#include "IUniversalObjectLocatorCustomization.h"
+
+#include "DragAndDrop/AssetDragDropOp.h"
+
+#include "PropertyCustomizationHelpers.h"
+
+#include "Widgets/Layout/SBox.h"
+
+
+#define LOCTEXT_NAMESPACE "AssetLocatorEditor"
+
+namespace UE::UniversalObjectLocator
+{
+
+bool FAssetLocatorEditor::IsDragSupported(TSharedPtr<FDragDropOperation> DragOperation, UObject* Context) const
+{
+	TSharedPtr<FAssetDragDropOp> ActorDrag;
+
+	if (DragOperation->IsOfType<FAssetDragDropOp>())
+	{
+		TSharedPtr<FAssetDragDropOp> AssetDrag = StaticCastSharedPtr<FAssetDragDropOp>(DragOperation);
+		if (AssetDrag->GetAssets().Num() == 1)
+		{
+			return true;
+		}
+	}
+
+	return false;
+}
+
+UObject* FAssetLocatorEditor::ResolveDragOperation(TSharedPtr<FDragDropOperation> DragOperation, UObject* Context) const
+{
+	TSharedPtr<FAssetDragDropOp> ActorDrag;
+
+	if (DragOperation->IsOfType<FAssetDragDropOp>())
+	{
+		TSharedPtr<FAssetDragDropOp> AssetDrag = StaticCastSharedPtr<FAssetDragDropOp>(DragOperation);
+		if (AssetDrag->GetAssets().Num() == 1)
+		{
+			return AssetDrag->GetAssets()[0].FastGetAsset(true);
+		}
+	}
+
+	return nullptr;
+}
+
+TSharedPtr<SWidget> FAssetLocatorEditor::MakeEditUI(TSharedPtr<IUniversalObjectLocatorCustomization> Customization)
+{
+	TSharedRef<SObjectPropertyEntryBox> EditWidget = SNew(SObjectPropertyEntryBox)
+	.ObjectPath(Customization.ToSharedRef(), &IUniversalObjectLocatorCustomization::GetPathToObject)
+	.PropertyHandle(Customization->GetProperty())
+	.AllowedClass(UObject::StaticClass())
+	.OnObjectChanged(this, &FAssetLocatorEditor::OnSetObject, TWeakPtr<IUniversalObjectLocatorCustomization>(Customization))
+	.AllowClear(true)
+	.DisplayUseSelected(true)
+	.DisplayBrowse(true)
+	.DisplayThumbnail(true);
+
+	float MinWidth = 100.f;
+	float MaxWidth = 500.f;
+	EditWidget->GetDesiredWidth(MinWidth, MaxWidth);
+
+	return SNew(SBox)
+	.MinDesiredWidth(MinWidth)
+	.MaxDesiredWidth(MaxWidth)
+	[
+		EditWidget
+	];
+}
+
+FText FAssetLocatorEditor::GetDisplayText() const
+{
+	return LOCTEXT("AssetLocatorName", "Asset");
+}
+
+FText FAssetLocatorEditor::GetDisplayTooltip() const
+{
+	return LOCTEXT("AssetLocatorTooltip", "Change this to an asset reference");
+}
+
+FSlateIcon FAssetLocatorEditor::GetDisplayIcon() const
+{
+	return FSlateIcon();
+}
+
+void FAssetLocatorEditor::OnSetObject(const FAssetData& InNewObject, TWeakPtr<IUniversalObjectLocatorCustomization> WeakCustomization)
+{
+	TSharedPtr<IUniversalObjectLocatorCustomization> Customization = WeakCustomization.Pin();
+	if (Customization)
+	{
+		// Assets are always absolute
+		UObject* Object = InNewObject.FastGetAsset(true);
+
+		FUniversalObjectLocator NewRef(Object, nullptr);
+		Customization->SetValue(MoveTemp(NewRef));
+	}
+}
+
+} // namespace UE::UniversalObjectLocator
+
+#undef LOCTEXT_NAMESPACE
diff --git a/Engine/Source/Editor/UniversalObjectLocatorEditor/Private/AssetLocatorEditor.h b/Engine/Source/Editor/UniversalObjectLocatorEditor/Private/AssetLocatorEditor.h
new file mode 100644
--- /dev/null
+++ b/Engine/Source/Editor/UniversalObjectLocatorEditor/Private/AssetLocatorEditor.h
@@ -0,0 +1,35 @@
+// Copyright Epic Games, Inc. All Rights Reserved.
+
+#pragma once
+
+#include "CoreTypes.h"
+#include "UniversalObjectLocatorEditor.h"
+
+struct FAssetData;
+
+namespace UE::UniversalObjectLocator
+{
+
+class IUniversalObjectLocatorCustomization;
+
+class FAssetLocatorEditor : public ILocatorEditor
+{
+	bool IsDragSupported(TSharedPtr<FDragDropOperation> DragOperation, UObject* Context) const override;
+
+	UObject* ResolveDragOperation(TSharedPtr<FDragDropOperation> DragOperation, UObject* Context) const override;
+
+	TSharedPtr<SWidget> MakeEditUI(TSharedPtr<IUniversalObjectLocatorCustomization> Customization) override;
+
+	FText GetDisplayText() const override;
+
+	FText GetDisplayTooltip() const override;
+
+	FSlateIcon GetDisplayIcon() const override;
+
+private:
+
+	void OnSetObject(const FAssetData& InNewObject, TWeakPtr<IUniversalObjectLocatorCustomization> WeakCustomization);
+};
+
+} // namespace UE::UniversalObjectLocator
+
diff --git a/Engine/Source/Editor/UniversalObjectLocatorEditor/Private/UniversalObjectLocatorCustomization.cpp b/Engine/Source/Editor/UniversalObjectLocatorEditor/Private/UniversalObjectLocatorCustomization.cpp
new file mode 100644
--- /dev/null
+++ b/Engine/Source/Editor/UniversalObjectLocatorEditor/Private/UniversalObjectLocatorCustomization.cpp
@@ -0,0 +1,507 @@
+// Copyright Epic Games, Inc. All Rights Reserved.
+
+#include "UniversalObjectLocatorCustomization.h"
+#include "UniversalObjectLocator.h"
+#include "UniversalObjectLocatorFragmentType.h"
+#include "UniversalObjectLocatorEditor.h"
+#include "UniversalObjectLocatorEditorModule.h"
+
+#include "Containers/Array.h"
+#include "Internationalization/Internationalization.h"
+#include "Modules/ModuleManager.h"
+
+#include "DetailWidgetRow.h"
+#include "IPropertyUtilities.h"
+#include "PropertyHandle.h"
+#include "PropertyCustomizationHelpers.h"
+
+#include "Styling/AppStyle.h"
+#include "Widgets/DeclarativeSyntaxSupport.h"
+#include "Widgets/Input/SComboButton.h"
+#include "Widgets/Text/STextBlock.h"
+#include "Framework/MultiBox/MultiBoxBuilder.h"
+#include "SDropTarget.h"
+
+#include "SceneOutlinerPublicTypes.h"
+#include "SceneOutlinerModule.h"
+
+#include "Async/Async.h"
+
+
+#define LOCTEXT_NAMESPACE "FUniversalObjectLocatorCustomization"
+
+namespace UE::UniversalObjectLocator
+{
+
+TSharedRef<IPropertyTypeCustomization> FUniversalObjectLocatorCustomization::MakeInstance()
+{
+	return MakeShared<FUniversalObjectLocatorCustomization>();
+}
+
+void FUniversalObjectLocatorCustomization::CustomizeHeader(TSharedRef<IPropertyHandle> StructPropertyHandle, FDetailWidgetRow& HeaderRow, IPropertyTypeCustomizationUtils& StructCustomizationUtils)
+{
+	PropertyHandle = StructPropertyHandle;
+	PropertyUtilities = StructCustomizationUtils.GetPropertyUtilities();
+
+	FSceneOutlinerInitializationOptions InitOptions;
+	{
+		InitOptions.bShowHeaderRow = false;
+	}
+
+	bool bDisplayThumbnail = true;
+
+	TArray<FString> AllowedTypes;
+	if (PropertyHandle->HasMetaData("AllowedLocators"))
+	{
+		PropertyHandle->GetMetaData("AllowedLocators").ParseIntoArray(AllowedTypes, TEXT(","));
+	}
+
+	FUniversalObjectLocatorEditorModule& Module = FModuleManager::Get().LoadModuleChecked<FUniversalObjectLocatorEditorModule>("UniversalObjectLocatorEditor");
+	for (TPair<FName, TSharedPtr<ILocatorEditor>> Pair : Module.LocatorEditors)
+	{
+		if (AllowedTypes.Num() == 0 || AllowedTypes.Contains(Pair.Key.ToString()))
+		{
+			ApplicableLocators.Add(Pair.Key, Pair.Value);
+		}
+	}
+
+	Content = SNew(SBox);
+	HeaderRow
+	.NameContent()
+	[
+		SNew(SHorizontalBox)
+
+		+ SHorizontalBox::Slot()
+		.AutoWidth()
+		[
+			StructPropertyHandle->CreatePropertyNameWidget()
+		]
+
+		// Type combo box
+		+ SHorizontalBox::Slot()
+		.Padding(FMargin(5.f, 0.f, 0.f, 0.f))
+		.AutoWidth()
+		[
+			SNew(SComboButton)
+			.Visibility(ApplicableLocators.Num() > 1 ? EVisibility::Visible : EVisibility::Collapsed)
+			.OnGetMenuContent(this, &FUniversalObjectLocatorCustomization::GetUserExposedFragmentTypeList)
+			.ButtonContent()
+			[
+				SNew(STextBlock)
+				.Text(this, &FUniversalObjectLocatorCustomization::GetCurrentFragmentTypeText)
+			]
+		]
+	]
+	.ValueContent()
+	.MaxDesiredWidth(TOptional<float>())
+	[
+		SNew(SDropTarget)
+		.OnAllowDrop(this, &FUniversalObjectLocatorCustomization::HandleIsDragAllowed)
+		.OnDropped(this, &FUniversalObjectLocatorCustomization::HandleDrop)
+		[
+			Content.ToSharedRef()
+		]
+	];
+
+	Rebuild();
+}
+
+void FUniversalObjectLocatorCustomization::Rebuild() const
+{
+	FUniversalObjectLocatorEditorModule& Module = FModuleManager::Get().LoadModuleChecked<FUniversalObjectLocatorEditorModule>("UniversalObjectLocatorEditor");
+	TSharedPtr<ILocatorEditor> LocatorEditor = ApplicableLocators.FindRef(GetCachedData().LocatorEditorType);
+	if (LocatorEditor)
+	{
+		TSharedPtr<SWidget> EditUI = LocatorEditor->MakeEditUI(SharedThis(const_cast<FUniversalObjectLocatorCustomization*>(this)));
+		if (EditUI)
+		{
+			TSharedPtr<SBox> BoxWidget = Content;
+			
+			// Delay reconstruction until next frame since this can be called from inside the previous EditUI,
+			// leading to destruction of the UI while it is still running
+			AsyncTask(ENamedThreads::GameThread, [BoxWidget, EditUI] {
+				BoxWidget->SetContent(EditUI.ToSharedRef());
+			});
+			return;
+		}
+	}
+
+	Content->SetContent(SNullWidget::NullWidget);
+}
+
+UObject* FUniversalObjectLocatorCustomization::GetSingleObject() const
+{
+	return GetCachedData().WeakObject.Get();
+}
+
+void FUniversalObjectLocatorCustomization::SetValue(FUniversalObjectLocator&& InNewValue) const
+{
+	PropertyHandle->NotifyPreChange();
+
+	PropertyHandle->EnumerateRawData([&InNewValue](void* RawData, int32 Index, int32 Num){
+		FUniversalObjectLocator* Reference = static_cast<FUniversalObjectLocator*>(RawData);
+		if (Index == Num-1)
+		{
+			*Reference = MoveTemp(InNewValue);
+		}
+		else
+		{
+			*Reference = InNewValue;
+		}
+		return true;
+	});
+
+	PropertyHandle->NotifyPostChange(EPropertyChangeType::ValueSet);
+	PropertyHandle->NotifyFinishedChangingProperties();
+}
+
+FString FUniversalObjectLocatorCustomization::GetPathToObject() const
+{
+	return GetCachedData().ObjectPath;
+}
+
+TSharedPtr<IPropertyHandle> FUniversalObjectLocatorCustomization::GetProperty() const
+{
+	return PropertyHandle;
+}
+
+TSharedRef<SWidget> FUniversalObjectLocatorCustomization::GetUserExposedFragmentTypeList()
+{
+	struct FMenuData
+	{
+		TSharedPtr<ILocatorEditor> LocatorEditor;
+		FText DisplayText;
+		FText DisplayTooltip;
+		FSlateIcon DisplayIcon;
+		FName LocatorEditorType;
+	};
+
+	TArray<FMenuData> MenuData;
+
+	for (TPair<FName, TSharedPtr<ILocatorEditor>> Pair : ApplicableLocators)
+	{
+		MenuData.Add(FMenuData{
+			Pair.Value,
+			Pair.Value->GetDisplayText(),
+			Pair.Value->GetDisplayTooltip(),
+			Pair.Value->GetDisplayIcon(),
+			Pair.Key
+		});
+	}
+
+	if (MenuData.Num() == 0)
+	{
+		// 
+		return SNullWidget::NullWidget;
+	}
+
+	Algo::SortBy(MenuData, &FMenuData::DisplayText, FText::FSortPredicate(ETextComparisonLevel::Default));
+
+	const bool bCloseAfterSelection = true;
+	FMenuBuilder MenuBuilder(bCloseAfterSelection, nullptr);
+
+	FCanExecuteAction AlwaysExecute = FCanExecuteAction::CreateLambda([]{ return true; });
+
+	MenuBuilder.BeginSection(NAME_None, LOCTEXT("FragmentTypeTypeHeader", "Change this locator's to reference a different type:"));
+	for (const FMenuData& Item : MenuData)
+	{
+		MenuBuilder.AddMenuEntry(
+			Item.DisplayText,
+			Item.DisplayTooltip,
+			Item.DisplayIcon,
+			FUIAction(
+				FExecuteAction::CreateSP(this, &FUniversalObjectLocatorCustomization::ChangeEditorType, Item.LocatorEditorType),
+				AlwaysExecute,
+				FIsActionChecked::CreateSP(this, &FUniversalObjectLocatorCustomization::CompareCurrentEditorType, Item.LocatorEditorType)
+			),
+			NAME_None,
+			EUserInterfaceActionType::RadioButton
+		);
+	}
+	MenuBuilder.EndSection();
+
+
+	return MenuBuilder.MakeWidget();
+}
+
+void FUniversalObjectLocatorCustomization::ChangeEditorType(FName InNewLocatorEditorType)
+{
+	PropertyHandle->NotifyPreChange();
+
+	PropertyHandle->EnumerateRawData(
+		[InNewLocatorEditorType](void* Data, int32 DataIndex, int32 Num)
+		{
+			FUniversalObjectLocator* Ref = static_cast<FUniversalObjectLocator*>(Data);
+			Ref->Reset();
+			return true;
+		}
+	);
+
+	CachedData.LocatorEditorType = InNewLocatorEditorType;
+
+	PropertyHandle->NotifyPostChange(EPropertyChangeType::ValueSet);
+	PropertyHandle->NotifyFinishedChangingProperties();
+
+	Rebuild();
+}
+
+bool FUniversalObjectLocatorCustomization::CompareCurrentEditorType(FName InLocatorEditorType) const
+{
+	return CachedData.LocatorEditorType == InLocatorEditorType;
+}
+
+FText FUniversalObjectLocatorCustomization::GetCurrentFragmentTypeText() const
+{
+	if (CachedData.LocatorEditorType != NAME_None)
+	{
+		return FText::FromName(CachedData.LocatorEditorType);
+	}
+
+	TArray<void*> RawData;
+	PropertyHandle->AccessRawData(RawData);
+
+	FText NoneText = LOCTEXT("NoValues", "None");
+	if (RawData.Num() == 0)
+	{
+		return NoneText;
+	}
+
+	const FFragmentType* CommonFragmentType = nullptr;
+	for (void* Ptr : RawData)
+	{
+		FUniversalObjectLocator* Reference    = static_cast<FUniversalObjectLocator*>(RawData[0]);
+		const FFragmentType*     FragmentType = Reference->GetLastFragmentType();
+
+		if (CommonFragmentType && FragmentType != CommonFragmentType)
+		{
+			return LOCTEXT("MultipleTypes", "<multiple>");
+		}
+		CommonFragmentType = FragmentType;
+	}
+	return CommonFragmentType ? CommonFragmentType->DisplayText : NoneText;
+}
+
+bool FUniversalObjectLocatorCustomization::HandleIsDragAllowed(TSharedPtr<FDragDropOperation> InDragOperation)
+{
+	for (TPair<FName, TSharedPtr<ILocatorEditor>> Pair : ApplicableLocators)
+	{
+		if (Pair.Value->IsDragSupported(InDragOperation, nullptr))
+		{
+			return true;
+		}
+	}
+
+	return false;
+}
+
+FReply FUniversalObjectLocatorCustomization::HandleDrop(const FGeometry& Geometry, const FDragDropEvent& DragDropEvent)
+{
+	TSharedPtr<FDragDropOperation> DragOperation = DragDropEvent.GetOperation();
+
+	UObject* ResolvedObject = nullptr;
+
+	for (TPair<FName, TSharedPtr<ILocatorEditor>> Pair : ApplicableLocators)
+	{
+		ResolvedObject = Pair.Value->ResolveDragOperation(DragOperation, nullptr);
+		if (ResolvedObject)
+		{
+			break;
+		}
+	}
+
+	if (!ResolvedObject)
+	{
+		return FReply::Unhandled();
+	}
+
+	TArray<void*> RawData;
+	PropertyHandle->AccessRawData(RawData);
+	for (void* Ptr : RawData)
+	{
+		FUniversalObjectLocator* Reference = static_cast<FUniversalObjectLocator*>(Ptr);
+		Reference->Reset(ResolvedObject);
+	}
+
+	return FReply::Handled();
+}
+
+void FUniversalObjectLocatorCustomization::SetActor(AActor* NewObject)
+{
+	FUniversalObjectLocator NewRef(NewObject);
+
+	TArray<void*> RawData;
+	PropertyHandle->AccessRawData(RawData);
+	for (void* Ptr : RawData)
+	{
+		FUniversalObjectLocator* Reference = static_cast<FUniversalObjectLocator*>(RawData[0]);
+		*Reference = NewRef;
+	}
+}
+
+FUniversalObjectLocator* FUniversalObjectLocatorCustomization::GetSinglePropertyValue()
+{
+	const FUniversalObjectLocator* ConstValue = static_cast<const FUniversalObjectLocatorCustomization*>(this)->GetSinglePropertyValue();
+	return const_cast<FUniversalObjectLocator*>(ConstValue);
+}
+
+const FUniversalObjectLocator* FUniversalObjectLocatorCustomization::GetSinglePropertyValue() const
+{
+	const FUniversalObjectLocator* Value = nullptr;
+
+	PropertyHandle->EnumerateConstRawData(
+		[&Value](const void* Data, int32 DataIndex, int32 Num)
+		{
+			if (Num == 1)
+			{
+				Value = static_cast<const FUniversalObjectLocator*>(Data);
+			}
+			return false;
+		}
+	);
+
+	return Value;
+}
+
+const FUniversalObjectLocatorCustomization::FCachedData& FUniversalObjectLocatorCustomization::GetCachedData() const
+{
+
+	const FUniversalObjectLocator* SingleValue = nullptr;
+
+	// Set the LocatorEditorType type based on the current values
+	TOptional<FName> CommonLocatorEditorType;
+
+	PropertyHandle->EnumerateConstRawData(
+		[this, &SingleValue, &CommonLocatorEditorType](const void* Data, int32 DataIndex, int32 Num)
+		{
+			const FUniversalObjectLocator* Value = static_cast<const FUniversalObjectLocator*>(Data);
+			if (Num == 1)
+			{
+				SingleValue = Value;
+			}
+			
+			const FFragmentType* FragmentTypePtr = Value->GetLastFragmentType();
+			if (!CommonLocatorEditorType)
+			{
+				if (FragmentTypePtr)
+				{
+					CommonLocatorEditorType = FragmentTypePtr->PrimaryEditorType;
+				}
+				else
+				{
+					// If we find any value that is effectively null, we choose the most applicable locator
+					return false;
+				}
+			}
+			else if (!FragmentTypePtr || FragmentTypePtr->PrimaryEditorType != CommonLocatorEditorType.GetValue())
+			{
+				CommonLocatorEditorType = NAME_None;
+			}
+			else
+			{
+				CommonLocatorEditorType = FragmentTypePtr->PrimaryEditorType;
+			}
+			return true;
+		}
+	);
+
+	if (CommonLocatorEditorType)
+	{
+		if (CommonLocatorEditorType.GetValue() != NAME_None)
+		{
+			CachedData.LocatorEditorType = CommonLocatorEditorType.GetValue();
+		}
+	}
+	else if (CachedData.LocatorEditorType == NAME_None)
+	{
+		for (const TPair<FName, TSharedPtr<ILocatorEditor>>& Pair : ApplicableLocators)
+		{
+			CachedData.LocatorEditorType = Pair.Key;
+			break;
+		}
+	}
+
+	bool bNeedsUpdate = false;
+	if (SingleValue)
+	{
+		if (!CachedData.PropertyValue.IsSet() || *SingleValue != CachedData.PropertyValue.GetValue())
+		{
+			CachedData.PropertyValue = *SingleValue;
+			bNeedsUpdate = true;
+		}
+	}
+	else if (CachedData.PropertyValue.IsSet())
+	{
+		CachedData.PropertyValue.Reset();
+		bNeedsUpdate = true;
+	}
+
+	if (bNeedsUpdate)
+	{
+		if (CachedData.PropertyValue.IsSet())
+		{
+			CachedData.WeakObject = CachedData.PropertyValue->SyncFind();
+		}
+		else
+		{
+			CachedData.WeakObject = nullptr;
+		}
+
+		if (UObject* Resolved = CachedData.WeakObject.Get())
+		{
+			CachedData.ObjectPath = FSoftObjectPath(Resolved).ToString();
+		}
+		else
+		{
+			CachedData.ObjectPath.Empty();
+		}
+
+		TArray<void*> RawData;
+		PropertyHandle->AccessRawData(RawData);
+
+		bool bAnyFragmentType = false;
+		const FFragmentType* CommonFragmentType = nullptr;
+		TOptional<FFragmentTypeHandle> CommonFragmentTypeHandle;
+		for (void* Ptr : RawData)
+		{
+			FUniversalObjectLocator* Locator      = static_cast<FUniversalObjectLocator*>(Ptr);
+			const FFragmentType*     FragmentType = Locator->GetLastFragmentType();
+
+			if (CommonFragmentType && FragmentType != CommonFragmentType)
+			{
+				CommonFragmentType = nullptr;
+				CommonFragmentTypeHandle.Reset();
+				break;
+			}
+
+			if (FragmentType)
+			{
+				bAnyFragmentType = true;
+			}
+			CommonFragmentType = FragmentType;
+			CommonFragmentTypeHandle = Locator->GetLastFragmentTypeHandle();
+		}
+
+		FText NoneText = LOCTEXT("NoValues", "None");
+		if (CommonFragmentType != nullptr)
+		{
+			CachedData.FragmentTypeText = CommonFragmentType->DisplayText;
+			CachedData.FragmentType     = CommonFragmentTypeHandle;
+		}
+		else if (!bAnyFragmentType)
+		{
+			CachedData.FragmentTypeText = NoneText;
+		}
+		else
+		{
+			CachedData.FragmentTypeText = LOCTEXT("MultipleTypes", "<multiple>");
+		}
+
+		Rebuild();
+	}
+
+	return CachedData;
+}
+
+} // namespace UE::UniversalObjectLocator
+
+#undef LOCTEXT_NAMESPACE
\ No newline at end of file
diff --git a/Engine/Source/Editor/UniversalObjectLocatorEditor/Private/UniversalObjectLocatorCustomization.h b/Engine/Source/Editor/UniversalObjectLocatorEditor/Private/UniversalObjectLocatorCustomization.h
new file mode 100644
--- /dev/null
+++ b/Engine/Source/Editor/UniversalObjectLocatorEditor/Private/UniversalObjectLocatorCustomization.h
@@ -0,0 +1,92 @@
+// Copyright Epic Games, Inc. All Rights Reserved.
+
+#pragma once
+
+#include "IPropertyTypeCustomization.h"
+#include "Internationalization/Text.h"
+#include "Templates/SharedPointer.h"
+#include "Misc/Optional.h"
+#include "Types/SlateEnums.h"
+
+#include "UniversalObjectLocator.h"
+#include "IUniversalObjectLocatorCustomization.h"
+
+struct FGeometry;
+struct FAssetData;
+
+class FReply;
+class FDragDropEvent;
+class FDetailWidgetRow;
+class FDragDropOperation;
+
+class IPropertyUtilities;
+class IDetailChildrenBuilder;
+
+class AActor;
+
+class SBox;
+class SWidget;
+
+namespace UE::UniversalObjectLocator
+{
+
+class ILocatorEditor;
+
+struct FUniversalObjectLocatorCustomization final : public IPropertyTypeCustomization, public IUniversalObjectLocatorCustomization
+{
+	static TSharedRef<IPropertyTypeCustomization> MakeInstance();
+
+	void CustomizeHeader(TSharedRef<IPropertyHandle> StructPropertyHandle, FDetailWidgetRow& HeaderRow, IPropertyTypeCustomizationUtils& StructCustomizationUtils) override;
+	void CustomizeChildren(TSharedRef<IPropertyHandle> StructPropertyHandle, IDetailChildrenBuilder& StructBuilder, IPropertyTypeCustomizationUtils& StructCustomizationUtils) override {}
+
+	UObject* GetContext() const override { return nullptr; }
+	UObject* GetSingleObject() const override;
+	FString GetPathToObject() const override;
+	void SetValue(FUniversalObjectLocator&& InNewValue) const override;
+	TSharedPtr<IPropertyHandle> GetProperty() const override;
+
+private:
+
+	struct FCachedData
+	{
+		TOptional<FUniversalObjectLocator> PropertyValue;
+		TWeakObjectPtr<> WeakObject;
+		FString ObjectPath;
+		TOptional<FFragmentTypeHandle> FragmentType;
+		FText FragmentTypeText;
+		FName LocatorEditorType;
+	};
+
+private:
+
+	void Rebuild() const;
+
+	TSharedRef<SWidget> GetUserExposedFragmentTypeList();
+	FText GetCurrentFragmentTypeText() const;
+
+	bool HandleIsDragAllowed(TSharedPtr<FDragDropOperation> InDragOperation);
+	FReply HandleDrop(const FGeometry& InGeometry, const FDragDropEvent& InDropEvent);
+
+	void SetActor(AActor* InActor);
+
+	const FCachedData& GetCachedData() const;
+
+	FUniversalObjectLocator* GetSinglePropertyValue();
+	const FUniversalObjectLocator* GetSinglePropertyValue() const;
+
+	void ChangeEditorType(FName InNewEditor);
+	bool CompareCurrentEditorType(FName InNewEditor) const;
+
+private:
+
+	TSharedPtr<IPropertyUtilities> PropertyUtilities;
+	TSharedPtr<IPropertyHandle> PropertyHandle;
+	TSharedPtr<SBox> Content;
+
+	mutable FCachedData CachedData;
+
+	TMap<FName, TSharedPtr<ILocatorEditor>> ApplicableLocators;
+};
+
+
+} // namespace UE::UniversalObjectLocator
\ No newline at end of file
diff --git a/Engine/Source/Editor/UniversalObjectLocatorEditor/Private/UniversalObjectLocatorEditorModule.cpp b/Engine/Source/Editor/UniversalObjectLocatorEditor/Private/UniversalObjectLocatorEditorModule.cpp
new file mode 100644
--- /dev/null
+++ b/Engine/Source/Editor/UniversalObjectLocatorEditor/Private/UniversalObjectLocatorEditorModule.cpp
@@ -0,0 +1,47 @@
+// Copyright Epic Games, Inc. All Rights Reserved.
+
+#include "UniversalObjectLocatorEditorModule.h"
+#include "Modules/ModuleManager.h"
+#include "PropertyEditorModule.h"
+#include "UniversalObjectLocatorCustomization.h"
+
+#include "AssetLocatorEditor.h"
+#include "ActorLocatorEditor.h"
+#include "AnimInstanceLocatorEditor.h"
+
+#include "ISequencerModule.h"
+
+namespace UE::UniversalObjectLocator
+{
+
+void FUniversalObjectLocatorEditorModule::StartupModule()
+{
+	FPropertyEditorModule& PropertyModule = FModuleManager::LoadModuleChecked<FPropertyEditorModule>("PropertyEditor");
+	PropertyModule.RegisterCustomPropertyTypeLayout("UniversalObjectLocator", FOnGetPropertyTypeCustomizationInstance::CreateStatic(&UE::UniversalObjectLocator::FUniversalObjectLocatorCustomization::MakeInstance));
+
+	// Register UOL editors
+	{
+		RegisterLocatorEditor("Actor", MakeShared<FActorLocatorEditor>());
+		RegisterLocatorEditor("AnimInstance", MakeShared<FAnimInstanceLocatorEditor>());
+		RegisterLocatorEditor("Asset", MakeShared<FAssetLocatorEditor>());
+	}
+}
+
+void FUniversalObjectLocatorEditorModule::ShutdownModule()
+{
+}
+
+void FUniversalObjectLocatorEditorModule::RegisterLocatorEditor(FName Name, TSharedPtr<ILocatorEditor> LocatorEditor)
+{
+	LocatorEditors.Add(Name, LocatorEditor);
+}
+
+void FUniversalObjectLocatorEditorModule::UnregisterLocatorEditor(FName Name)
+{
+	LocatorEditors.Remove(Name);
+}
+
+} // namespace UE::UniversalObjectLocator
+
+IMPLEMENT_MODULE(UE::UniversalObjectLocator::FUniversalObjectLocatorEditorModule, UniversalObjectLocatorEditor);
+
diff --git a/Engine/Source/Editor/UniversalObjectLocatorEditor/Private/UniversalObjectLocatorEditorModule.h b/Engine/Source/Editor/UniversalObjectLocatorEditor/Private/UniversalObjectLocatorEditorModule.h
new file mode 100644
--- /dev/null
+++ b/Engine/Source/Editor/UniversalObjectLocatorEditor/Private/UniversalObjectLocatorEditorModule.h
@@ -0,0 +1,34 @@
+// Copyright Epic Games, Inc. All Rights Reserved.
+
+#pragma once
+
+#include "UniversalObjectLocatorFwd.h"
+#include "IUniversalObjectLocatorEditorModule.h"
+#include "Containers/Map.h"
+#include "Templates/SharedPointer.h"
+#include "UObject/NameTypes.h"
+
+namespace UE::UniversalObjectLocator
+{
+
+class ILocatorEditor;
+
+class FUniversalObjectLocatorEditorModule
+	: public IUniversalObjectLocatorEditorModule
+{
+public:
+
+	void StartupModule() override;
+
+	void ShutdownModule() override;
+
+	void RegisterLocatorEditor(FName LocatorName, TSharedPtr<ILocatorEditor> LocatorEditor) override;
+
+	void UnregisterLocatorEditor(FName LocatorName) override;
+
+	TMap<FName, TSharedPtr<ILocatorEditor>> LocatorEditors;
+};
+
+} // namespace UE::UniversalObjectLocator
+
+
diff --git a/Engine/Source/Editor/UniversalObjectLocatorEditor/Public/IUniversalObjectLocatorCustomization.h b/Engine/Source/Editor/UniversalObjectLocatorEditor/Public/IUniversalObjectLocatorCustomization.h
new file mode 100644
--- /dev/null
+++ b/Engine/Source/Editor/UniversalObjectLocatorEditor/Public/IUniversalObjectLocatorCustomization.h
@@ -0,0 +1,30 @@
+// Copyright Epic Games, Inc. All Rights Reserved.
+
+#pragma once
+
+#include "Templates/SharedPointer.h"
+
+struct FUniversalObjectLocator;
+
+class FString;
+class UObject;
+class IPropertyHandle;
+
+namespace UE::UniversalObjectLocator
+{
+
+class IUniversalObjectLocatorCustomization
+{
+public:
+	virtual ~IUniversalObjectLocatorCustomization() = default;
+
+	virtual UObject* GetContext() const = 0;
+	virtual UObject* GetSingleObject() const = 0;
+	virtual FString GetPathToObject() const = 0;
+	virtual void SetValue(FUniversalObjectLocator&& InNewValue) const = 0;
+
+	virtual TSharedPtr<IPropertyHandle> GetProperty() const = 0;
+};
+
+
+} // namespace UE::UniversalObjectLocator
\ No newline at end of file
diff --git a/Engine/Source/Editor/UniversalObjectLocatorEditor/Public/IUniversalObjectLocatorEditorModule.h b/Engine/Source/Editor/UniversalObjectLocatorEditor/Public/IUniversalObjectLocatorEditorModule.h
new file mode 100644
--- /dev/null
+++ b/Engine/Source/Editor/UniversalObjectLocatorEditor/Public/IUniversalObjectLocatorEditorModule.h
@@ -0,0 +1,25 @@
+// Copyright Epic Games, Inc. All Rights Reserved.
+
+#pragma once
+
+#include "Modules/ModuleInterface.h"
+#include "Templates/SharedPointerFwd.h"
+
+class FName;
+
+namespace UE::UniversalObjectLocator
+{
+
+class ILocatorEditor;
+
+class IUniversalObjectLocatorEditorModule
+	: public IModuleInterface
+{
+public:
+
+	virtual void RegisterLocatorEditor(FName Name, TSharedPtr<ILocatorEditor> LocatorEditor) = 0;
+	virtual void UnregisterLocatorEditor(FName Name) = 0;
+};
+
+
+} // namespace UE::UniversalObjectLocator
\ No newline at end of file
diff --git a/Engine/Source/Editor/UniversalObjectLocatorEditor/Public/UniversalObjectLocatorEditor.h b/Engine/Source/Editor/UniversalObjectLocatorEditor/Public/UniversalObjectLocatorEditor.h
new file mode 100644
--- /dev/null
+++ b/Engine/Source/Editor/UniversalObjectLocatorEditor/Public/UniversalObjectLocatorEditor.h
@@ -0,0 +1,36 @@
+// Copyright Epic Games, Inc. All Rights Reserved.
+
+#pragma once
+
+#include "Templates/SharedPointer.h"
+
+struct FSlateIcon;
+
+class FText;
+class UObject;
+class SWidget;
+class FDragDropOperation;
+
+namespace UE::UniversalObjectLocator
+{
+
+class IUniversalObjectLocatorCustomization;
+
+class ILocatorEditor : public TSharedFromThis<ILocatorEditor>
+{
+public:
+	virtual ~ILocatorEditor() = default;
+
+	virtual bool IsDragSupported(TSharedPtr<FDragDropOperation> DragOperation, UObject* Context) const = 0;
+
+	virtual UObject* ResolveDragOperation(TSharedPtr<FDragDropOperation> DragOperation, UObject* Context) const = 0;
+
+	virtual TSharedPtr<SWidget> MakeEditUI(TSharedPtr<IUniversalObjectLocatorCustomization> Customization) = 0;
+
+	virtual FText GetDisplayText() const = 0;
+	virtual FText GetDisplayTooltip() const = 0;
+	virtual FSlateIcon GetDisplayIcon() const = 0;
+};
+
+
+} // namespace UE::UniversalObjectLocator
\ No newline at end of file
diff --git a/Engine/Source/Editor/UniversalObjectLocatorEditor/UniversalObjectLocatorEditor.Build.cs b/Engine/Source/Editor/UniversalObjectLocatorEditor/UniversalObjectLocatorEditor.Build.cs
new file mode 100644
--- /dev/null
+++ b/Engine/Source/Editor/UniversalObjectLocatorEditor/UniversalObjectLocatorEditor.Build.cs
@@ -0,0 +1,33 @@
+// Copyright Epic Games, Inc. All Rights Reserved.
+
+using UnrealBuildTool;
+
+public class UniversalObjectLocatorEditor : ModuleRules
+{
+	public UniversalObjectLocatorEditor(ReadOnlyTargetRules Target) : base(Target)
+	{
+		UnsafeTypeCastWarningLevel = WarningLevel.Error;
+
+		PublicDependencyModuleNames.AddRange(
+			new string[] {
+				"PropertyEditor",
+				"SlateCore",
+			}
+		);
+
+		PrivateDependencyModuleNames.AddRange(
+			new string[] {
+				"Core",
+				"CoreUObject",
+				"EditorWidgets",
+				"Engine",
+				"PropertyEditor",
+				"SceneOutliner",
+				"Slate",
+				"SlateCore",
+				"Sequencer",
+				"UniversalObjectLocator",
+			}
+		);
+	}
+}
diff --git a/Engine/Source/Editor/UnrealEd/Private/EditorEngine.cpp b/Engine/Source/Editor/UnrealEd/Private/EditorEngine.cpp
--- a/Engine/Source/Editor/UnrealEd/Private/EditorEngine.cpp
+++ b/Engine/Source/Editor/UnrealEd/Private/EditorEngine.cpp
@@ -1312,6 +1312,7 @@ void UEditorEngine::LoadDefaultEditorModules()
 			TEXT("AnimationSettings"),
 			TEXT("GameplayDebuggerEditor"),
 			TEXT("RenderResourceViewer"),
+			TEXT("UniversalObjectLocatorEditor"),
 		};
 
 	FScopedSlowTask ModuleSlowTask((float)UE_ARRAY_COUNT(ModuleNames));
diff --git a/Engine/Source/Editor/UnrealEd/UnrealEd.Build.cs b/Engine/Source/Editor/UnrealEd/UnrealEd.Build.cs
--- a/Engine/Source/Editor/UnrealEd/UnrealEd.Build.cs
+++ b/Engine/Source/Editor/UnrealEd/UnrealEd.Build.cs
@@ -183,6 +183,7 @@ public class UnrealEd : ModuleRules
 				"FieldNotification",
 				"VirtualizationEditor",
 				"GeometryCore",
+				"UniversalObjectLocatorEditor",
 			}
 		);
 
diff --git a/Engine/Source/Runtime/Engine/Engine.Build.cs b/Engine/Source/Runtime/Engine/Engine.Build.cs
--- a/Engine/Source/Runtime/Engine/Engine.Build.cs
+++ b/Engine/Source/Runtime/Engine/Engine.Build.cs
@@ -157,6 +157,7 @@ public class Engine : ModuleRules
 				"TraceLog",
 				"ColorManagement",
 				"Icmp",
+				"UniversalObjectLocator",
 				"XmlParser",
 			}
 		);
diff --git a/Engine/Source/Runtime/Engine/Private/UniversalObjectLocators/ActorLocatorFragment.cpp b/Engine/Source/Runtime/Engine/Private/UniversalObjectLocators/ActorLocatorFragment.cpp
new file mode 100644
--- /dev/null
+++ b/Engine/Source/Runtime/Engine/Private/UniversalObjectLocators/ActorLocatorFragment.cpp
@@ -0,0 +1,144 @@
+// Copyright Epic Games, Inc. All Rights Reserved.
+
+#include "UniversalObjectLocators/ActorLocatorFragment.h"
+#include "UniversalObjectLocators/IActorLocatorFragmentResolver.h"
+#include "UniversalObjectLocatorFragmentTypeHandle.h"
+#include "UniversalObjectLocatorResolveParams.h"
+#include "UniversalObjectLocatorStringParams.h"
+#include "UniversalObjectLocatorInitializeParams.h"
+#include "UniversalObjectLocatorInitializeResult.h"
+
+#include "UObject/UnrealNames.h"
+#include "LevelUtils.h"
+#include "Engine/Level.h"
+#include "Engine/LevelStreaming.h"
+#include "Engine/World.h"
+#include "Misc/EditorPathHelper.h"
+
+UE::UniversalObjectLocator::TFragmentTypeHandle<FActorLocatorFragment> FActorLocatorFragment::FragmentType;
+
+ULevel* GetLevelFromContext(const UObject* InContext)
+{
+	ULevel* Level = Cast<ULevel>(const_cast<UObject*>(InContext));
+	if (!Level && InContext)
+	{
+		Level = InContext->GetTypedOuter<ULevel>();
+	}
+	return Level;
+}
+
+UObject* ResolveActorWithinLevel(const FActorLocatorFragment& Payload, ULevel* Level)
+{
+	// Default to owning world (to resolve AlwaysLoaded actors not part of a Streaming Level and Disabled Streaming World Partitions)
+	UWorld* StreamingWorld = nullptr;
+
+	// Construct the path to the level asset that the streamed level relates to
+	ULevelStreaming* LevelStreaming = ULevelStreaming::FindStreamingLevel(Level);
+	if (LevelStreaming)
+	{
+		// If we're loading a world partition runtime cell, we need to find the streaming world that is responsible for resolving those actors
+		if (Level->IsWorldPartitionRuntimeCell())
+		{
+			StreamingWorld = LevelStreaming->GetStreamingWorld();
+			check(StreamingWorld);
+			LevelStreaming = ULevelStreaming::FindStreamingLevel(StreamingWorld->PersistentLevel);
+		}
+		else
+		{
+			StreamingWorld = Level->GetTypedOuter<UWorld>();
+		}
+	}
+
+	if (LevelStreaming && StreamingWorld)
+	{
+		// StreamedLevelPackage is a package name of the form /Game/Folder/MapName, not a full asset path
+		FName StreamedPackageName = (LevelStreaming->PackageNameToLoad == NAME_None) ? LevelStreaming->GetWorldAssetPackageFName() : LevelStreaming->PackageNameToLoad;
+
+		// @todo: we're only checking package name here - to be 100% correct we should really check the asset name as well,
+		//        but that is probably not necessary because multiple level assets in a single package are not supported
+		if (Payload.Path.GetAssetPath().GetPackageName() == StreamedPackageName)
+		{
+			// Payload.Path.GetSubPathString() specifies the path from the package (so includes PersistentLevel.) so we must do a ResolveSubObject from its outer
+
+			UObject* ResolvedObject = nullptr;
+			StreamingWorld->ResolveSubobject(*Payload.Path.GetSubPathString(), ResolvedObject, /*bLoadIfExists*/false);
+			return ResolvedObject;
+		}
+	}
+
+	return nullptr;
+}
+
+UE::UniversalObjectLocator::FResolveResult FActorLocatorFragment::Resolve(const UE::UniversalObjectLocator::FResolveParams& Params) const
+{
+	using namespace UE::UniversalObjectLocator;
+
+	UObject* Result = nullptr;
+
+	// Handle the context being a IActorLocatorFragmentResolver object - this is the first port of call
+	if (const IActorLocatorFragmentResolver* Resolver = Cast<const IActorLocatorFragmentResolver>(Params.Context))
+	{
+		if (Resolver && Resolver->ResolveActorLocatorPayload(*this, Result))
+		{
+			return FResolveResultData(Result);
+		}
+	}
+
+	// Next handle default level streaming and partition worlds behavior
+	if (ULevel* Level = GetLevelFromContext(Params.Context))
+	{
+		Result = ResolveActorWithinLevel(*this, Level);
+	}
+
+	// Finally fallback to just trying to resolve the path directly
+	if (!Result)
+	{
+		Result = Path.ResolveObject();
+	}
+
+	return FResolveResultData(Result);
+}
+
+void FActorLocatorFragment::ToString(FStringBuilderBase& OutStringBuilder) const
+{
+	Path.AppendString(OutStringBuilder);
+}
+
+UE::UniversalObjectLocator::FParseStringResult FActorLocatorFragment::TryParseString(FStringView InString, const UE::UniversalObjectLocator::FParseStringParams& Params)
+{
+	Path = InString;
+	return UE::UniversalObjectLocator::FParseStringResult().Success();
+}
+
+UE::UniversalObjectLocator::FInitializeResult FActorLocatorFragment::Initialize(const UE::UniversalObjectLocator::FInitializeParams& InParams)
+{
+	using namespace UE::UniversalObjectLocator;
+
+#if WITH_EDITOR
+	if (InParams.Context)
+	{
+		Path = FEditorPathHelper::GetEditorPathFromReferencer(InParams.Object, InParams.Context);
+	}
+	else
+	{
+		Path = FEditorPathHelper::GetEditorPath(InParams.Object);
+	}
+#else
+	Path = InParams.Object;
+#endif
+
+	// Really, actors should be relative to their level in order to support streaming within level instances, but
+	//   world partition makes that impossible
+	return FInitializeResult::Absolute();
+}
+
+uint32 FActorLocatorFragment::ComputePriority(const UObject* ObjectToReference, const UObject* Context)
+{
+	// Can only reference actors and components
+	if (ObjectToReference->IsA<AActor>())
+	{
+		// Let internal references be used first if possible
+		return 100;
+	}
+	return 0;
+}
diff --git a/Engine/Source/Runtime/Engine/Private/UniversalObjectLocators/AnimInstanceLocatorFragment.cpp b/Engine/Source/Runtime/Engine/Private/UniversalObjectLocators/AnimInstanceLocatorFragment.cpp
new file mode 100644
--- /dev/null
+++ b/Engine/Source/Runtime/Engine/Private/UniversalObjectLocators/AnimInstanceLocatorFragment.cpp
@@ -0,0 +1,158 @@
+// Copyright Epic Games, Inc. All Rights Reserved.
+
+#include "UniversalObjectLocators/AnimInstanceLocatorFragment.h"
+#include "UniversalObjectLocator.h"
+#include "UniversalObjectLocatorInitializeParams.h"
+#include "UniversalObjectLocatorInitializeResult.h"
+#include "UniversalObjectLocatorFragmentTypeHandle.h"
+#include "UniversalObjectLocatorResolveParams.h"
+#include "UniversalObjectLocatorStringParams.h"
+
+#include "Animation/AnimInstance.h"
+#include "Components/SkeletalMeshComponent.h"
+
+#define LOCTEXT_NAMESPACE "AnimInstanceLocatorFragment"
+
+namespace UE::UniversalObjectLocator
+{
+	static constexpr FStringView DefaultAnimInstance(TEXTVIEW("default"));
+	static constexpr FStringView PostProcessAnimInstance(TEXTVIEW("post_process"));
+}
+
+UE::UniversalObjectLocator::TFragmentTypeHandle<FAnimInstanceLocatorFragment> FAnimInstanceLocatorFragment::FragmentType;
+
+uint32 FAnimInstanceLocatorFragment::ComputePriority(const UObject* ObjectToReference, const UObject* Context)
+{
+	if (Context && Context->IsA<USkeletalMeshComponent>())
+	{
+		return 1000;
+	}
+	if (ObjectToReference && (ObjectToReference->IsA<USkeletalMeshComponent>() || ObjectToReference->IsA<UAnimInstance>()))
+	{
+		return 1000;
+	}
+	return 0;
+}
+
+UE::UniversalObjectLocator::FResolveResult FAnimInstanceLocatorFragment::Resolve(const UE::UniversalObjectLocator::FResolveParams& Params) const
+{
+	const USkeletalMeshComponent* SkelMeshComp = Cast<const USkeletalMeshComponent>(Params.Context);
+
+	if (SkelMeshComp)
+	{
+		if (Type == EAnimInstanceLocatorFragmentType::PostProcessAnimInstance)
+		{
+			return UE::UniversalObjectLocator::FResolveResultData(SkelMeshComp->GetPostProcessInstance());
+		}
+		else
+		{
+			return UE::UniversalObjectLocator::FResolveResultData(SkelMeshComp->GetAnimInstance());
+		}
+	}
+
+	return UE::UniversalObjectLocator::FResolveResult();
+}
+
+UE::UniversalObjectLocator::FInitializeResult FAnimInstanceLocatorFragment::Initialize(const UE::UniversalObjectLocator::FInitializeParams& InParams)
+{
+	using namespace UE::UniversalObjectLocator;
+
+	// 3 Different options here:
+	// Context is a USkeletalMeshComponent
+	// 		Result: Payload with a Type as either AnimInstance or PostProcessAnimInstance depending on what Object is
+	// Object is a USkeletalMeshComponent, Context is nullptr
+	// 		Result: ParentContext points to the skel mesh component, Type depends on UAnimInstance
+	// Object is a UAnimInstance, Context is nullptr
+	// 		Result: ParentContext points to the skel mesh component, Type depends on UAnimInstance
+
+	const USkeletalMeshComponent* SkelMeshComp = InParams.GetContextAs<USkeletalMeshComponent>();
+	if (SkelMeshComp)
+	{
+		const UAnimInstance* AnimInstance = InParams.GetObjectAs<const UAnimInstance>();
+		if (AnimInstance && AnimInstance == SkelMeshComp->GetPostProcessInstance())
+		{
+			// Use the post process instance
+			Type = EAnimInstanceLocatorFragmentType::PostProcessAnimInstance;
+		}
+		else
+		{
+			// Assume main anim instance
+			Type = EAnimInstanceLocatorFragmentType::AnimInstance;
+		}
+
+		return FInitializeResult::Relative(SkelMeshComp);
+	}
+	else
+	{
+		// Context needs to be baked into the locator
+		const UAnimInstance* AnimInstance = InParams.GetObjectAs<const UAnimInstance>();
+		if (AnimInstance)
+		{
+			SkelMeshComp = AnimInstance->GetTypedOuter<const USkeletalMeshComponent>();
+		}
+		else
+		{
+			SkelMeshComp = InParams.GetObjectAs<const USkeletalMeshComponent>();
+		}
+
+		if (!ensureMsgf(SkelMeshComp, TEXT("Unable to create an Anim Instance locator for object %s since there is no valid skeletal mesh."), InParams.Object ? *InParams.Object->GetName() : TEXT("<<None>>")))
+		{
+			return FInitializeResult::Failure();
+		}
+
+		if (AnimInstance && AnimInstance == SkelMeshComp->GetPostProcessInstance())
+		{
+			// Use the post process instance
+			Type = EAnimInstanceLocatorFragmentType::PostProcessAnimInstance;
+		}
+		{
+			// Assume main anim instance
+			Type = EAnimInstanceLocatorFragmentType::AnimInstance;
+		}
+
+		return FInitializeResult::Relative(SkelMeshComp);
+	}
+}
+
+void FAnimInstanceLocatorFragment::ToString(FStringBuilderBase& OutStringBuilder) const
+{
+	using namespace UE::UniversalObjectLocator;
+
+	switch (Type)
+	{
+	case EAnimInstanceLocatorFragmentType::AnimInstance:
+		// Leave blank for default
+		// OutStringBuilder.Append(DefaultAnimInstance);
+		break;
+	case EAnimInstanceLocatorFragmentType::PostProcessAnimInstance:
+		OutStringBuilder.Append(PostProcessAnimInstance);
+		break;
+	}
+}
+
+UE::UniversalObjectLocator::FParseStringResult FAnimInstanceLocatorFragment::TryParseString(FStringView InString, const UE::UniversalObjectLocator::FParseStringParams& Params)
+{
+	using namespace UE::UniversalObjectLocator;
+
+	if (InString.Len() == 0 || InString.Compare(DefaultAnimInstance, ESearchCase::IgnoreCase) == 0)
+	{
+		Type = EAnimInstanceLocatorFragmentType::AnimInstance;
+		return FParseStringResult().Success();
+	}
+	else if (InString.Compare(PostProcessAnimInstance, ESearchCase::IgnoreCase) == 0)
+	{
+		Type = EAnimInstanceLocatorFragmentType::AnimInstance;
+		return FParseStringResult().Success();
+	}
+
+	return FParseStringResult().Failure(
+		UE_UOL_PARSE_ERROR(Params,
+			FText::Format(
+				LOCTEXT("Error_InvalidAnimInstanceType", "Invalid anim instance type '{0}'."),
+				FText::FromStringView(InString)
+			)
+		)
+	);
+}
+
+#undef LOCTEXT_NAMESPACE
diff --git a/Engine/Source/Runtime/Engine/Private/UniversalObjectLocators/AssetLocatorFragment.cpp b/Engine/Source/Runtime/Engine/Private/UniversalObjectLocators/AssetLocatorFragment.cpp
new file mode 100644
--- /dev/null
+++ b/Engine/Source/Runtime/Engine/Private/UniversalObjectLocators/AssetLocatorFragment.cpp
@@ -0,0 +1,92 @@
+// Copyright Epic Games, Inc. All Rights Reserved.
+
+#include "UniversalObjectLocators/AssetLocatorFragment.h"
+#include "UniversalObjectLocatorInitializeParams.h"
+#include "UniversalObjectLocatorInitializeResult.h"
+#include "UniversalObjectLocatorFragmentTypeHandle.h"
+#include "UniversalObjectLocatorResolveParams.h"
+#include "UniversalObjectLocatorStringParams.h"
+#include "UObject/ObjectRedirector.h"
+#include "UObject/Package.h"
+
+UE::UniversalObjectLocator::TFragmentTypeHandle<FAssetLocatorFragment> FAssetLocatorFragment::FragmentType;
+
+UE::UniversalObjectLocator::FResolveResult FAssetLocatorFragment::Resolve(const UE::UniversalObjectLocator::FResolveParams& Params) const
+{
+	using namespace UE::UniversalObjectLocator;
+
+	FString PathString = Path.ToString();
+
+	if (EnumHasAnyFlags(Params.Flags, EResolveFlags::Unload))
+	{
+		// Assets cannot be explicitly unloaded
+		return FResolveResult();
+	}
+
+	FResolveResultData Result;
+
+	Result.Object = FindObject<UObject>(nullptr, *PathString);
+	if (Result.Object == nullptr && EnumHasAnyFlags(Params.Flags, EResolveFlags::Load))
+	{
+		Result.Object = StaticLoadObject(UObject::StaticClass(), nullptr, *PathString, nullptr, LOAD_None, nullptr, true);
+		Result.Flags.bWasLoaded = (Result.Object != nullptr);
+	}
+
+#if WITH_EDITOR
+	// Look at core redirects if we didn't find the object
+	if (!Result.Object)
+	{
+		FSoftObjectPath FixupObjectPath(MoveTemp(PathString));
+		if (FixupObjectPath.FixupCoreRedirects())
+		{
+			FString FixedUpPathString = FixupObjectPath.ToString();
+
+			Result.Object = FindObject<UObject>(nullptr, *FixedUpPathString);
+			if (Result.Object == nullptr && EnumHasAnyFlags(Params.Flags, EResolveFlags::Load))
+			{
+				Result.Object = LoadObject<UObject>(nullptr, *FixedUpPathString);
+				Result.Flags.bWasLoaded = (Result.Object != nullptr);
+			}
+		}
+	}
+#endif
+
+	while (UObjectRedirector* Redirector = Cast<UObjectRedirector>(Result.Object))
+	{
+		Result.Object = Redirector->DestinationObject;
+	}
+
+	return Result;
+}
+
+UE::UniversalObjectLocator::FInitializeResult FAssetLocatorFragment::Initialize(const UE::UniversalObjectLocator::FInitializeParams& InParams)
+{
+	using namespace UE::UniversalObjectLocator;
+
+	Path.TrySetPath(InParams.Object);
+	return FInitializeResult::Absolute();
+}
+
+void FAssetLocatorFragment::ToString(FStringBuilderBase& OutStringBuilder) const
+{
+	Path.AppendString(OutStringBuilder);
+}
+
+UE::UniversalObjectLocator::FParseStringResult FAssetLocatorFragment::TryParseString(FStringView InString, const UE::UniversalObjectLocator::FParseStringParams& Params)
+{
+	Path = InString;
+	return UE::UniversalObjectLocator::FParseStringResult().Success();
+}
+
+uint32 FAssetLocatorFragment::ComputePriority(const UObject* ObjectToReference, const UObject* Context)
+{
+	// Can only reference top level persistent assets
+	UObject* Outer = ObjectToReference->GetOuter();
+	if (Outer && Outer->IsA<UPackage>() && !ObjectToReference->HasAnyFlags(RF_Transient) && Outer != GetTransientPackage())
+	{
+		return 1000;
+	}
+
+	// We can't use this at all
+	return 0;
+}
diff --git a/Engine/Source/Runtime/Engine/Private/UnrealEngine.cpp b/Engine/Source/Runtime/Engine/Private/UnrealEngine.cpp
--- a/Engine/Source/Runtime/Engine/Private/UnrealEngine.cpp
+++ b/Engine/Source/Runtime/Engine/Private/UnrealEngine.cpp
@@ -128,6 +128,8 @@ UnrealEngine.cpp: Implements the UEngine class and helpers.
 #include "Iris/IrisConfig.h"
 #endif
 
+#include "IUniversalObjectLocatorModule.h"
+
 #include "Particles/Spawn/ParticleModuleSpawn.h"
 #include "Particles/TypeData/ParticleModuleTypeDataMesh.h"
 #include "Particles/ParticleLODLevel.h"
@@ -137,6 +139,10 @@ UnrealEngine.cpp: Implements the UEngine class and helpers.
 #include "Rendering/SkeletalMeshRenderData.h"
 #include "IO/IoDispatcher.h"
 
+#include "UniversalObjectLocators/ActorLocatorFragment.h"
+#include "UniversalObjectLocators/AssetLocatorFragment.h"
+#include "UniversalObjectLocators/AnimInstanceLocatorFragment.h"
+
 #if WITH_EDITOR
 #include "Settings/LevelEditorPlaySettings.h"
 #include "LandscapeSubsystem.h"
@@ -2029,6 +2035,29 @@ void UEngine::Init(IEngineLoop* InEngineLoop)
 
 	InitializeObjectReferences();
 
+	// Initialize built-in UOL fragment types
+	{
+		using namespace UE::UniversalObjectLocator;
+
+		IUniversalObjectLocatorModule& UolModule = FModuleManager::Get().LoadModuleChecked<IUniversalObjectLocatorModule>("UniversalObjectLocator");
+
+		{
+			FFragmentTypeParameters FragmentTypeParams("actor", NSLOCTEXT("Engine", "ActorLocatorFragment", "Actor"));
+			FragmentTypeParams.PrimaryEditorType = "Actor";
+			FActorLocatorFragment::FragmentType = UolModule.RegisterFragmentType<FActorLocatorFragment>(FragmentTypeParams);
+		}
+		{
+			FFragmentTypeParameters FragmentTypeParams("asset", NSLOCTEXT("Engine", "AssetLocatorFragment", "Asset"));
+			FragmentTypeParams.PrimaryEditorType = "Asset";
+			FAssetLocatorFragment::FragmentType = UolModule.RegisterFragmentType<FAssetLocatorFragment>(FragmentTypeParams);
+		}
+		{
+			FFragmentTypeParameters FragmentTypeParams("animinst", NSLOCTEXT("Engine", "AnimInstanceLocatorFragment", "AnimInstance"));
+			FragmentTypeParams.PrimaryEditorType = "Component";
+			FAnimInstanceLocatorFragment::FragmentType = UolModule.RegisterFragmentType<FAnimInstanceLocatorFragment>(FragmentTypeParams);
+		}
+	}
+
 	if (GConfig)
 	{
 		bool bTemp = true;
diff --git a/Engine/Source/Runtime/Engine/Public/UniversalObjectLocators/ActorLocatorFragment.h b/Engine/Source/Runtime/Engine/Public/UniversalObjectLocators/ActorLocatorFragment.h
new file mode 100644
--- /dev/null
+++ b/Engine/Source/Runtime/Engine/Public/UniversalObjectLocators/ActorLocatorFragment.h
@@ -0,0 +1,48 @@
+// Copyright Epic Games, Inc. All Rights Reserved.
+
+#pragma once
+
+#include "CoreTypes.h"
+#include "UniversalObjectLocatorFwd.h"
+#include "ActorLocatorFragment.generated.h"
+
+/**
+ * 32 Bytes (40 in-editor).
+ */
+USTRUCT()
+struct FActorLocatorFragment
+{
+	GENERATED_BODY()
+
+	UPROPERTY()
+	FSoftObjectPath Path;
+
+	ENGINE_API static UE::UniversalObjectLocator::TFragmentTypeHandle<FActorLocatorFragment> FragmentType;
+
+	friend uint32 GetTypeHash(const FActorLocatorFragment& A)
+	{
+		return GetTypeHash(A.Path);
+	}
+
+	friend bool operator==(const FActorLocatorFragment& A, const FActorLocatorFragment& B)
+	{
+		return A.Path == B.Path;
+	}
+
+	UE::UniversalObjectLocator::FResolveResult Resolve(const UE::UniversalObjectLocator::FResolveParams& Params) const;
+	UE::UniversalObjectLocator::FInitializeResult Initialize(const UE::UniversalObjectLocator::FInitializeParams& InParams);
+
+	void ToString(FStringBuilderBase& OutStringBuilder) const;
+	UE::UniversalObjectLocator::FParseStringResult TryParseString(FStringView InString, const UE::UniversalObjectLocator::FParseStringParams& Params);
+
+	static uint32 ComputePriority(const UObject* Object, const UObject* Context);
+};
+
+template<>
+struct TStructOpsTypeTraits<FActorLocatorFragment> : public TStructOpsTypeTraitsBase2<FActorLocatorFragment>
+{
+	enum
+	{
+		WithIdenticalViaEquality = true,
+	};
+};
\ No newline at end of file
diff --git a/Engine/Source/Runtime/Engine/Public/UniversalObjectLocators/AnimInstanceLocatorFragment.h b/Engine/Source/Runtime/Engine/Public/UniversalObjectLocators/AnimInstanceLocatorFragment.h
new file mode 100644
--- /dev/null
+++ b/Engine/Source/Runtime/Engine/Public/UniversalObjectLocators/AnimInstanceLocatorFragment.h
@@ -0,0 +1,65 @@
+// Copyright Epic Games, Inc. All Rights Reserved.
+
+#pragma once
+
+#include "CoreTypes.h"
+#include "UniversalObjectLocatorFwd.h"
+#include "AnimInstanceLocatorFragment.generated.h"
+
+UENUM()
+enum class EAnimInstanceLocatorFragmentType
+{
+	AnimInstance,
+	PostProcessAnimInstance
+};
+
+/**
+ * 32 Bytes (40 in-editor).
+ */
+USTRUCT()
+struct FAnimInstanceLocatorFragment
+{
+	GENERATED_BODY()
+
+	UPROPERTY()
+	EAnimInstanceLocatorFragmentType Type;
+
+	FAnimInstanceLocatorFragment()
+		: Type(EAnimInstanceLocatorFragmentType::AnimInstance)
+	{
+	}
+
+	explicit FAnimInstanceLocatorFragment(EAnimInstanceLocatorFragmentType InType)
+		: Type(InType)
+	{
+	}
+
+	ENGINE_API static UE::UniversalObjectLocator::TFragmentTypeHandle<FAnimInstanceLocatorFragment> FragmentType;
+
+	friend uint32 GetTypeHash(const FAnimInstanceLocatorFragment& A)
+	{
+		return GetTypeHash(A.Type);
+	}
+
+	friend bool operator==(const FAnimInstanceLocatorFragment& A, const FAnimInstanceLocatorFragment& B)
+	{
+		return A.Type == B.Type;
+	}
+
+	UE::UniversalObjectLocator::FResolveResult Resolve(const UE::UniversalObjectLocator::FResolveParams& Params) const;
+	UE::UniversalObjectLocator::FInitializeResult Initialize(const UE::UniversalObjectLocator::FInitializeParams& InParams);
+
+	void ToString(FStringBuilderBase& OutStringBuilder) const;
+	UE::UniversalObjectLocator::FParseStringResult TryParseString(FStringView InString, const UE::UniversalObjectLocator::FParseStringParams& Params);
+
+	static uint32 ComputePriority(const UObject* Object, const UObject* Context);
+};
+
+template<>
+struct TStructOpsTypeTraits<FAnimInstanceLocatorFragment> : public TStructOpsTypeTraitsBase2<FAnimInstanceLocatorFragment>
+{
+	enum
+	{
+		WithIdenticalViaEquality = true,
+	};
+};
\ No newline at end of file
diff --git a/Engine/Source/Runtime/Engine/Public/UniversalObjectLocators/AssetLocatorFragment.h b/Engine/Source/Runtime/Engine/Public/UniversalObjectLocators/AssetLocatorFragment.h
new file mode 100644
--- /dev/null
+++ b/Engine/Source/Runtime/Engine/Public/UniversalObjectLocators/AssetLocatorFragment.h
@@ -0,0 +1,48 @@
+// Copyright Epic Games, Inc. All Rights Reserved.
+
+#pragma once
+
+#include "CoreTypes.h"
+#include "UniversalObjectLocatorFwd.h"
+#include "AssetLocatorFragment.generated.h"
+
+/**
+ * 32 Bytes (40 in-editor).
+ */
+USTRUCT()
+struct FAssetLocatorFragment
+{
+	GENERATED_BODY()
+
+	UPROPERTY()
+	FTopLevelAssetPath Path;
+
+	friend uint32 GetTypeHash(const FAssetLocatorFragment& A)
+	{
+		return GetTypeHash(A.Path);
+	}
+
+	friend bool operator==(const FAssetLocatorFragment& A, const FAssetLocatorFragment& B)
+	{
+		return A.Path == B.Path;
+	}
+
+	ENGINE_API static UE::UniversalObjectLocator::TFragmentTypeHandle<FAssetLocatorFragment> FragmentType;
+
+	UE::UniversalObjectLocator::FResolveResult Resolve(const UE::UniversalObjectLocator::FResolveParams& Params) const;
+	UE::UniversalObjectLocator::FInitializeResult Initialize(const UE::UniversalObjectLocator::FInitializeParams& InParams);
+
+	void ToString(FStringBuilderBase& OutStringBuilder) const;
+	UE::UniversalObjectLocator::FParseStringResult TryParseString(FStringView InString, const UE::UniversalObjectLocator::FParseStringParams& Params);
+
+	static uint32 ComputePriority(const UObject* Object, const UObject* Context);
+};
+
+template<>
+struct TStructOpsTypeTraits<FAssetLocatorFragment> : public TStructOpsTypeTraitsBase2<FAssetLocatorFragment>
+{
+	enum
+	{
+		WithIdenticalViaEquality = true,
+	};
+};
\ No newline at end of file
diff --git a/Engine/Source/Runtime/Engine/Public/UniversalObjectLocators/IActorLocatorFragmentResolver.h b/Engine/Source/Runtime/Engine/Public/UniversalObjectLocators/IActorLocatorFragmentResolver.h
new file mode 100644
--- /dev/null
+++ b/Engine/Source/Runtime/Engine/Public/UniversalObjectLocators/IActorLocatorFragmentResolver.h
@@ -0,0 +1,25 @@
+// Copyright Epic Games, Inc. All Rights Reserved.
+
+#pragma once
+
+#include "CoreTypes.h"
+#include "UObject/Interface.h"
+#include "IActorLocatorFragmentResolver.generated.h"
+
+struct FActorLocatorFragment;
+
+UINTERFACE(MinimalAPI)
+class UActorLocatorFragmentResolver
+	: public UInterface
+{
+public:
+	GENERATED_BODY()
+};
+
+class IActorLocatorFragmentResolver
+{
+public:
+	GENERATED_BODY()
+
+	virtual bool ResolveActorLocatorPayload(const FActorLocatorFragment& Payload, UObject*& OutResult) const = 0;
+};
diff --git a/Engine/Source/Runtime/UniversalObjectLocator/Private/DirectPathObjectLocator.cpp b/Engine/Source/Runtime/UniversalObjectLocator/Private/DirectPathObjectLocator.cpp
new file mode 100644
--- /dev/null
+++ b/Engine/Source/Runtime/UniversalObjectLocator/Private/DirectPathObjectLocator.cpp
@@ -0,0 +1,57 @@
+// Copyright Epic Games, Inc. All Rights Reserved.
+
+#include "DirectPathObjectLocator.h"
+#include "UniversalObjectLocatorInitializeParams.h"
+#include "UniversalObjectLocatorInitializeResult.h"
+#include "UniversalObjectLocatorFragmentTypeHandle.h"
+#include "UniversalObjectLocatorResolveParams.h"
+#include "UniversalObjectLocatorStringParams.h"
+
+UE::UniversalObjectLocator::TFragmentTypeHandle<FDirectPathObjectLocator> FDirectPathObjectLocator::FragmentType;
+
+UE::UniversalObjectLocator::FResolveResult FDirectPathObjectLocator::Resolve(const UE::UniversalObjectLocator::FResolveParams& Params) const
+{
+	using namespace UE::UniversalObjectLocator;
+
+	if (UObject* Object = Path.ResolveObject())
+	{
+		return FResolveResultData(Object);
+	}
+
+	if (EnumHasAnyFlags(Params.Flags, EResolveFlags::Load))
+	{
+		UObject* Object = Path.TryLoad();
+
+		FResolveResultData Result(Object);
+		Result.Flags.bWasLoaded = Object != nullptr;
+		return Result;
+	}
+
+	return FResolveResult();
+}
+
+UE::UniversalObjectLocator::FInitializeResult FDirectPathObjectLocator::Initialize(const UE::UniversalObjectLocator::FInitializeParams& InParams)
+{
+	using namespace UE::UniversalObjectLocator;
+
+	Path.SetPath(InParams.Object->GetPathName());
+	return FInitializeResult::Absolute();
+}
+
+void FDirectPathObjectLocator::ToString(FStringBuilderBase& OutStringBuilder) const
+{
+	Path.AppendString(OutStringBuilder);
+}
+
+UE::UniversalObjectLocator::FParseStringResult FDirectPathObjectLocator::TryParseString(FStringView InString, const UE::UniversalObjectLocator::FParseStringParams& Params)
+{
+	Path = InString;
+	return UE::UniversalObjectLocator::FParseStringResult().Success();
+}
+
+uint32 FDirectPathObjectLocator::ComputePriority(const UObject* ObjectToReference, const UObject* Context)
+{
+	// We can't use this at all unless explicitly told to
+	return 0;
+}
+
diff --git a/Engine/Source/Runtime/UniversalObjectLocator/Private/SubObjectLocator.cpp b/Engine/Source/Runtime/UniversalObjectLocator/Private/SubObjectLocator.cpp
new file mode 100644
--- /dev/null
+++ b/Engine/Source/Runtime/UniversalObjectLocator/Private/SubObjectLocator.cpp
@@ -0,0 +1,70 @@
+// Copyright Epic Games, Inc. All Rights Reserved.
+
+#include "SubObjectLocator.h"
+#include "UniversalObjectLocatorInitializeParams.h"
+#include "UniversalObjectLocatorInitializeResult.h"
+#include "UniversalObjectLocatorFragmentTypeHandle.h"
+#include "UniversalObjectLocatorResolveParams.h"
+#include "UniversalObjectLocatorStringParams.h"
+
+UE::UniversalObjectLocator::TFragmentTypeHandle<FSubObjectLocator> FSubObjectLocator::FragmentType;
+
+UE::UniversalObjectLocator::FResolveResult FSubObjectLocator::Resolve(const UE::UniversalObjectLocator::FResolveParams& Params) const
+{
+	using namespace UE::UniversalObjectLocator;
+
+	if (UObject* NonConstContext = const_cast<UObject*>(Params.Context))
+	{
+		UObject* Result = nullptr;
+		if (!NonConstContext->ResolveSubobject(*PathWithinContext, Result, false))
+		{
+			Result = FindObject<UObject>(NonConstContext, *PathWithinContext);
+		}
+		return FResolveResultData(Result);
+	}
+	return FResolveResult();
+}
+
+UE::UniversalObjectLocator::FInitializeResult FSubObjectLocator::Initialize(const UE::UniversalObjectLocator::FInitializeParams& InParams)
+{
+	using namespace UE::UniversalObjectLocator;
+
+	PathWithinContext.Reset();
+	if (ensureMsgf(InParams.Context && InParams.Object->IsIn(InParams.Context),
+		TEXT("Unable to create a reference to %s from context %s since the object does not exist within the context"),
+		*InParams.Object->GetName(),
+		InParams.Context ? *InParams.Context->GetName() : TEXT("<<none>>"))
+		)
+	{
+		InParams.Object->GetPathName(InParams.Context, PathWithinContext);
+
+		return FInitializeResult::Relative(InParams.Context);
+	}
+
+	return FInitializeResult::Failure();
+}
+
+void FSubObjectLocator::ToString(FStringBuilderBase& OutStringBuilder) const
+{
+	OutStringBuilder.Append(PathWithinContext);
+}
+
+UE::UniversalObjectLocator::FParseStringResult FSubObjectLocator::TryParseString(FStringView InString, const UE::UniversalObjectLocator::FParseStringParams& Params)
+{
+	PathWithinContext.Reset(InString.Len());
+	PathWithinContext.Append(InString.GetData(), InString.Len());
+	return UE::UniversalObjectLocator::FParseStringResult().Success();
+}
+
+uint32 FSubObjectLocator::ComputePriority(const UObject* ObjectToReference, const UObject* Context)
+{
+	// Can only reference objects that are relative to the context
+	if (Context && ObjectToReference->IsIn(Context))
+	{
+		return 2000;
+	}
+
+	// We can't use this at all
+	return 0;
+}
+
diff --git a/Engine/Source/Runtime/UniversalObjectLocator/Private/UniversalObjectLocator.cpp b/Engine/Source/Runtime/UniversalObjectLocator/Private/UniversalObjectLocator.cpp
new file mode 100644
--- /dev/null
+++ b/Engine/Source/Runtime/UniversalObjectLocator/Private/UniversalObjectLocator.cpp
@@ -0,0 +1,670 @@
+// Copyright Epic Games, Inc. All Rights Reserved.
+
+#include "UniversalObjectLocator.h"
+#include "DirectPathObjectLocator.h"
+#include "UniversalObjectLocatorInitializeParams.h"
+#include "UniversalObjectLocatorInitializeResult.h"
+#include "UniversalObjectLocatorStringParams.h"
+#include "UniversalObjectLocatorStringUtils.h"
+
+#define LOCTEXT_NAMESPACE "UOL"
+
+namespace UE::UniversalObjectLocator
+{
+	static constexpr FStringView MagicLeadingString = TEXTVIEW("uobj://");
+
+	extern TArray<FFragmentType> GFragmentTypes;
+
+	const FFragmentType* FindBestFragmentType(const UObject* Object, const UObject* Context);
+
+} // namespace UE::UniversalObjectLocator
+
+bool operator==(const FUniversalObjectLocator& A, const FUniversalObjectLocator& B)
+{
+	return A.Fragments == B.Fragments;
+}
+
+bool operator!=(const FUniversalObjectLocator& A, const FUniversalObjectLocator& B)
+{
+	return A.Fragments != B.Fragments;
+}
+
+uint32 GetTypeHash(const FUniversalObjectLocator& Locator)
+{
+	return GetTypeHash(Locator.Fragments);
+}
+
+FUniversalObjectLocator::FUniversalObjectLocator()
+{
+
+}
+
+FUniversalObjectLocator::FUniversalObjectLocator(UObject* Object, const UObject* Context, const UObject* StopAtContext)
+{
+	Reset(Object, Context, StopAtContext);
+}
+
+UE::UniversalObjectLocator::FResolveResult FUniversalObjectLocator::Resolve(const FResolveParams& Params) const
+{
+	using namespace UE::UniversalObjectLocator;
+
+	// Check for invalid combinations of flags
+	check(!EnumHasAllFlags(Params.Flags, EResolveFlags::Load | EResolveFlags::Unload));
+	// Cannot have WillWait without Async
+	check(!EnumHasAllFlags(Params.Flags, EResolveFlags::WillWait) || EnumHasAllFlags(Params.Flags, EResolveFlags::Async));
+
+	if (EnumHasAnyFlags(Params.Flags, EResolveFlags::Async))
+	{
+		return ResolveAsyncImpl(Params);
+	}
+	else
+	{
+		return ResolveSyncImpl(Params);
+	}
+}
+
+UE::UniversalObjectLocator::FResolveResult FUniversalObjectLocator::ResolveSyncImpl(const FResolveParams& Params) const
+{
+	using namespace UE::UniversalObjectLocator;
+
+	check(!EnumHasAnyFlags(Params.Flags, EResolveFlags::Async));
+
+	FResolveResult EmptyResult;
+
+	if (Fragments.Num() == 0)
+	{
+		return EmptyResult;
+	}
+
+	if (Fragments.Num() == 1)
+	{
+		return Fragments[0].Resolve(Params);
+	}
+
+	bool bLoadedIndirectly = false;
+	
+	FResolveResult LastResult;
+
+	const UObject* CurrentContext = Params.Context;
+	const int32 Num = Fragments.Num();
+	for (int32 Index = 0; Index < Num; ++Index)
+	{
+		const FUniversalObjectLocatorFragment& Fragment = Fragments[Index];
+
+		const bool bLastFragment = Index == (Num-1);
+
+		// Only unload the last one
+		FResolveParams RelativeParams(CurrentContext, Params.Flags);
+		if (!bLastFragment)
+		{
+			RelativeParams.Flags &= ~EResolveFlags::Unload;
+		}
+
+		LastResult = Fragment.Resolve(RelativeParams);
+
+		if (EnumHasAnyFlags(RelativeParams.Flags, EResolveFlags::Unload))
+		{
+			return LastResult;
+		}
+
+		FResolveResultData ResultData = LastResult.SyncGet();
+		if (ResultData.Object == nullptr)
+		{
+			// If anything fails to resolve, nothing resolves
+			return EmptyResult;
+		}
+
+		CurrentContext = ResultData.Object;
+
+		// If the last one was implicitly loaded or created, the final result should report that
+		
+		if (bLastFragment)
+		{
+			ResultData.Flags.bWasLoadedIndirectly = bLoadedIndirectly;
+		}
+		else
+		{
+			bLoadedIndirectly |= ResultData.Flags.bWasLoaded;
+		}
+	}
+
+	return LastResult;
+}
+
+UE::UniversalObjectLocator::FResolveResult FUniversalObjectLocator::ResolveAsyncImpl(const FResolveParams& Params) const
+{
+	using namespace UE::UniversalObjectLocator;
+
+	check(EnumHasAnyFlags(Params.Flags, EResolveFlags::Async));
+
+	if (Fragments.Num() == 0)
+	{
+		return FResolveResult();
+	}
+
+	struct FState : TSharedFromThis<FState>
+	{
+		FState(const TArray<FUniversalObjectLocatorFragment>& InFragments, EResolveFlags InInputResolveFlags)
+			: FragmentsCopy(InFragments)
+			, CurrentIndex(-1)
+			, InputResolveFlags(InInputResolveFlags)
+		{}
+
+		void ProcessNext(FResolveResultData LastResult)
+		{
+			const int32 Index = ++CurrentIndex;
+
+			const bool bFinished       = Index == FragmentsCopy.Num();
+			const bool bLastFragment   = Index == (FragmentsCopy.Num()-1);
+
+			UObject* Context = LastResult.Object;
+
+			// If the previous one was loaded or created, any subsequent operations are loaded indirectly
+			if (!bLastFragment)
+			{
+				bLoadedIndirectly = LastResult.Flags.bWasLoaded;
+			}
+
+			const bool bCannotContinue = Context == nullptr && Index != 0;
+			if (bFinished || bCannotContinue)
+			{
+				if (AsyncResult.IsSet())
+				{
+					LastResult.Flags.bWasLoadedIndirectly = bLoadedIndirectly;
+					AsyncResult->SetValue(LastResult);
+				}
+				else
+				{
+					FinalResult = FResolveResult(LastResult);
+				}
+				return;
+			}
+
+			// Try and resolve this one
+			FResolveParams RelativeParams(Context, InputResolveFlags);
+
+			// Only unload the last one
+			if (!bLastFragment)
+			{
+				RelativeParams.Flags &= ~EResolveFlags::Unload;
+			}
+
+			FResolveResult Result = FragmentsCopy[Index].Resolve(RelativeParams);
+
+			// If we don't need to wait, call the next one immediately
+			if (!Result.NeedsWait())
+			{
+				ProcessNext(Result.SyncGet());
+				return;
+			}
+
+			// We need to wait for something to complete...
+			//   If the currently held FinalResult is not async, we need to make it so
+			if (!FinalResult.IsAsync())
+			{
+				check(!AsyncResult.IsSet());
+				AsyncResult.Emplace();
+				FinalResult = FResolveResult(AsyncResult->GetFuture());
+			}
+
+			// Set the continuation
+			Result.AsyncGet(
+				[State = AsShared()](const FResolveResultData& InValue)
+				{
+					State->ProcessNext(InValue);
+				}
+			);
+		}
+
+		TArray<FUniversalObjectLocatorFragment> FragmentsCopy;
+
+		TOptional<
+			TPromise<FResolveResultData>
+		> AsyncResult;
+
+		FResolveResult FinalResult;
+
+		int32 CurrentIndex;
+		EResolveFlags InputResolveFlags;
+		bool bLoadedIndirectly = false;
+	};
+
+	TSharedPtr<FState> SharedState = MakeShared<FState>(this->Fragments, Params.Flags);
+	SharedState->ProcessNext(nullptr);
+	return MoveTemp(SharedState->FinalResult);
+}
+
+UObject* FUniversalObjectLocator::SyncFind(const UObject* Context) const
+{
+	return Resolve(FResolveParams::SyncFind(Context)).SyncGet().Object;
+}
+
+UObject* FUniversalObjectLocator::SyncLoad(const UObject* Context) const
+{
+	return Resolve(FResolveParams::SyncLoad(Context)).SyncGet().Object;
+}
+
+void FUniversalObjectLocator::SyncUnload(const UObject* Context) const
+{
+	Resolve(FResolveParams::SyncUnload(Context)).SyncGet();
+}
+
+UE::UniversalObjectLocator::FResolveResult FUniversalObjectLocator::AsyncFind(const UObject* Context) const
+{
+	return Resolve(FResolveParams::AsyncFind(Context));
+}
+
+UE::UniversalObjectLocator::FResolveResult FUniversalObjectLocator::AsyncLoad(const UObject* Context) const
+{
+	return Resolve(FResolveParams::AsyncLoad(Context));
+}
+
+UE::UniversalObjectLocator::FResolveResult FUniversalObjectLocator::AsyncUnload(const UObject* Context) const
+{
+	return Resolve(FResolveParams::AsyncUnload(Context));
+}
+
+void FUniversalObjectLocator::Reset()
+{
+	Fragments.Empty();
+}
+
+void FUniversalObjectLocator::Reset(UObject* InObject, const UObject* Context, const UObject* StopAtContext)
+{
+	Fragments.Reset();
+	if (!InObject || !AddFragment(InObject, Context, StopAtContext))
+	{
+		// Failed to create the locator
+		Fragments.Empty();
+	}
+}
+
+const UE::UniversalObjectLocator::FFragmentType* FUniversalObjectLocator::GetLastFragmentType() const
+{
+	return Fragments.Num() != 0 ? Fragments.Last().GetFragmentType() : nullptr;
+}
+
+UE::UniversalObjectLocator::FFragmentTypeHandle FUniversalObjectLocator::GetLastFragmentTypeHandle() const
+{
+	return Fragments.Num() != 0 ? Fragments.Last().GetFragmentTypeHandle() : UE::UniversalObjectLocator::FFragmentTypeHandle();
+}
+
+FUniversalObjectLocatorFragment* FUniversalObjectLocator::GetLastFragment()
+{
+	return Fragments.Num() != 0 ? &Fragments.Last() : nullptr;
+}
+
+const FUniversalObjectLocatorFragment* FUniversalObjectLocator::GetLastFragment() const
+{
+	return Fragments.Num() != 0 ? &Fragments.Last() : nullptr;
+}
+
+void FUniversalObjectLocator::ToString(FStringBuilderBase& OutString) const
+{
+	using namespace UE::UniversalObjectLocator;
+
+	// Universal Object Locators currently write to strings of the following form:
+	// uobj://fragment-type-id=payload-string!...!fragment-type-id-n=payload-string-n
+
+	TStringBuilder<128> PayloadScratchSpace;
+
+	OutString += MagicLeadingString;
+
+	int32 NumPrintedTypes = 0;
+	int32 NumPrintedPayloads = 0;
+
+	// Print the fragment types as the path
+	for (int32 Index = 0; Index < Fragments.Num(); ++Index)
+	{
+		const FUniversalObjectLocatorFragment& Fragment = Fragments[Index];
+		if (const FFragmentType* FragmentTypePtr = Fragment.GetFragmentType())
+		{
+			if (NumPrintedTypes != 0)
+			{
+				OutString += '/';
+			}
+
+			FragmentTypePtr->FragmentTypeID.AppendString(OutString);
+			++NumPrintedTypes;
+		}
+	}
+
+	if (NumPrintedTypes == 0)
+	{
+		OutString += TEXT("none");
+		return;
+	}
+
+	OutString += TEXT("?");
+
+	// Print the payloads as a query string
+	for (int32 Index = 0; Index < Fragments.Num(); ++Index)
+	{
+		const FUniversalObjectLocatorFragment& Fragment = Fragments[Index];
+
+		const FFragmentType* FragmentTypePtr = Fragment.GetFragmentType();
+		const UStruct*       FragmentStruct  = FragmentTypePtr ? FragmentTypePtr->GetStruct() : nullptr;
+
+		if (FragmentStruct)
+		{
+			if (NumPrintedPayloads != 0)
+			{
+				OutString += '&';
+			}
+
+			PayloadScratchSpace.Reset();
+			FragmentTypePtr->ToString(Fragment.GetPayload(), PayloadScratchSpace);
+			if (PayloadScratchSpace.Len() != 0)
+			{
+				OutString.Appendf(TEXT("payload%i="), Index);
+				OutString.Append(PayloadScratchSpace.ToView());
+			}
+
+			++NumPrintedPayloads;
+		}
+	}
+}
+
+UE::UniversalObjectLocator::FParseStringResult FUniversalObjectLocator::TryParseString(FStringView InString, const FParseStringParams& InParams)
+{
+	using namespace UE::UniversalObjectLocator;
+
+	FParseStringResult Result;
+
+	if (!InString.StartsWith(MagicLeadingString, ESearchCase::IgnoreCase))
+	{
+		return Result.Failure(UE_UOL_PARSE_ERROR(InParams, LOCTEXT("Error_MissingScheme", "String does not start with uobj://")));
+	}
+
+	InString = Result.Progress(InString, MagicLeadingString.Len());
+	if (InString.Len() == 0)
+	{
+		return Result.Failure(UE_UOL_PARSE_ERROR(InParams, LOCTEXT("Error_EmptyPath", "Path is empty.")));
+	}
+
+	if (InString.Compare(TEXTVIEW("none"), ESearchCase::IgnoreCase) == 0)
+	{
+		Reset();
+		return Result.Success();
+	}
+
+	FUniversalObjectLocator Tmp;
+
+	// First parse the fragment types
+	bool bFinishedParsingFragments = false;
+	while (InString.Len() != 0 && !bFinishedParsingFragments)
+	{
+		// Find the next / or ? character
+		const int32 NextDelimiter = UE::String::FindFirstOfAnyChar(InString, TEXTVIEW("/?#"));
+
+		FStringView ThisFragment = InString;
+		if (NextDelimiter != INDEX_NONE)
+		{
+			// If the delimiter is a ? this is the last fragment, but there might be
+			//   some additional payload string to parse
+			bFinishedParsingFragments = (InString[NextDelimiter] == '?');
+
+			// Trim the fragment string to the delimiter that we found
+			//   and move the remaining string to after the delimiter
+			ThisFragment = InString.Left(NextDelimiter);
+
+			// Now trim the rest of the string to remove this fragment.
+			//   If the delimiter is a # then we have nothing more to parse
+			if (InString[NextDelimiter] == '#')
+			{
+				// Indicate that we finished parsing at the #, but reset the string
+				//   so that we stop parsing anything else
+				Result.NumCharsParsed += NextDelimiter;
+				InString.Reset();
+			}
+			else
+			{
+				// Progress past the delimiter
+				InString = Result.Progress(InString, NextDelimiter + 1);
+			}
+
+			// Empty string - just move on to the next one
+			if (ThisFragment.Len() == 0)
+			{
+				continue;
+			}
+		}
+		else
+		{
+			// No delimiter found, so the rest of the string is parsed as a fragment type
+			Result.NumCharsParsed += InString.Len();
+			InString.Reset();
+		}
+
+		FUniversalObjectLocatorFragment& NewFragment = Tmp.Fragments.Emplace_GetRef();
+
+		FParseStringResult TypeResult = NewFragment.TryParseFragmentType(ThisFragment, InParams);
+		Result.NumCharsParsed += TypeResult.NumCharsParsed;
+		if (!TypeResult)
+		{
+			return Result.Failure(
+				UE_UOL_PARSE_ERROR(
+					InParams,
+					FText::Format(
+						LOCTEXT("Error_InvalidFragmentType", "Fragment Type specifier is invalid: {0}."),
+						FText::FromStringView(ThisFragment)
+					)
+				)
+			);
+		}
+	}
+
+	// Next parse the fragment payloads
+	bool bFinishedParsingPayloads = false;
+	while (InString.Len() != 0 && !bFinishedParsingPayloads)
+	{
+		// Find the next & or # character
+		const int32 NextDelimiter = UE::String::FindFirstOfAnyChar(InString, TEXTVIEW("&#"));
+
+		FStringView ThisPayload = InString;
+		if (NextDelimiter != INDEX_NONE)
+		{
+			// If the delimiter is a # this is the last payload, so we finish parsing
+			//   after this one
+			bFinishedParsingFragments = (InString[NextDelimiter] == '#');
+
+			// Trim the fragment string to the delimiter that we found
+			//   and move the remaining string to after the delimiter
+			ThisPayload = InString.Left(NextDelimiter);
+
+			// Progress past the delimiter
+			InString = Result.Progress(InString, NextDelimiter + 1);
+
+			// Empty string - just move on to the next one
+			if (ThisPayload.Len() == 0)
+			{
+				continue;
+			}
+		}
+		else
+		{
+			// No delimiter - parse the remaining string as a payload and finish parsing
+			Result.NumCharsParsed += InString.Len();
+			InString.Reset();
+		}
+
+		static constexpr FStringView PayloadString = TEXTVIEW("payload");
+
+		// Try and parse a fragment index - any other query string key=value pairs will be skipped
+		const int32 EqualsDelimiter = UE::String::FindFirstChar(ThisPayload, '=');
+		if (EqualsDelimiter != INDEX_NONE && ThisPayload.StartsWith(PayloadString, ESearchCase::IgnoreCase) && ThisPayload.Len() > 8)
+		{
+			int32 NumChars = 0;
+			uint32 ParsedIndex = 0;
+			// Parse an integer from the position after the 'payload' string
+			if (!ParseUnsignedInteger(ThisPayload.RightChop(PayloadString.Len()), ParsedIndex, &NumChars) || ParsedIndex > static_cast<uint32>(std::numeric_limits<int32>::max()))
+			{
+				// Invalid int
+				Result.NumCharsParsed += ThisPayload.Len();
+				continue;
+			}
+			else if (ParsedIndex >= static_cast<uint32>(Tmp.Fragments.Num()))
+			{
+				// Invalid index
+				Result.NumCharsParsed += ThisPayload.Len();
+				continue;
+			}
+
+			// Handle non-sensical payload strings like payload0123another=
+			if (PayloadString.Len() + NumChars != EqualsDelimiter)
+			{
+				Result.NumCharsParsed += ThisPayload.Len();
+				continue;
+			}
+
+			// Skip over the index and =
+			ThisPayload = Result.Progress(ThisPayload, PayloadString.Len() + NumChars + 1);
+
+			if (ThisPayload.Len() == 0)
+			{
+				// Leave the fragment as default if there is an empty string
+				continue;
+			}
+
+			const int32 FragmentIndex = static_cast<int32>(ParsedIndex);
+			if (!Tmp.Fragments.IsValidIndex(FragmentIndex))
+			{
+				// Invalid fragment index for this payload - skip it
+				continue;
+			}
+
+			FUniversalObjectLocatorFragment& Fragment = Tmp.Fragments[FragmentIndex];
+
+			FParseStringResult PayloadResult = Fragment.TryParseFragmentPayload(ThisPayload, InParams);
+			Result.NumCharsParsed += PayloadResult.NumCharsParsed;
+			if (!PayloadResult)
+			{
+				const FFragmentType* Type = Fragment.GetFragmentType();
+
+				return Result.Failure(
+					UE_UOL_PARSE_ERROR(InParams,
+						FText::Format(
+							LOCTEXT("Error_InvalidPayload", "Payload is invalid for fragment type {0}: {1}."),
+							FText::FromName(Type ? Type->FragmentTypeID : NAME_None),
+							FText::FromStringView(ThisPayload)
+						)
+					)
+				);
+			}
+		}
+	}
+
+	*this = MoveTemp(Tmp);
+	return Result.Success();
+}
+
+FUniversalObjectLocator FUniversalObjectLocator::FromString(FStringView InString, const FParseStringParams& InParams)
+{
+	FUniversalObjectLocator Locator;
+	Locator.TryParseString(InString, InParams);
+	return Locator;
+}
+
+bool FUniversalObjectLocator::AddFragment(const UObject* Object, const UObject* Context, const UObject* StopAtContext)
+{
+	using namespace UE::UniversalObjectLocator;
+
+	const FFragmentType* FragmentType = FindBestFragmentType(Object, Context);
+	if (!FragmentType)
+	{
+		return false;
+	}
+
+	// Initialize the payload
+	FUniversalObjectLocatorFragment NewLocator(*FragmentType);
+	FInitializeResult Result = FragmentType->InitializePayload(NewLocator.GetPayload(), FInitializeParams { Object, Context });
+
+	if (Result.Type == ELocatorType::Undefined)
+	{
+		return false;
+	}
+
+	// If the initialization needs to be relative to a different context, add a fragment for NewContext as well
+	if (Result.Type == ELocatorType::Relative && Result.RelativeToContext != Context)
+	{
+		if (ensureMsgf(Result.RelativeToContext, TEXT("Payload initialization reported a relative locator but did not specify what it is relative to.")))
+		{
+			if (StopAtContext != Result.RelativeToContext)
+			{
+				AddFragment(Result.RelativeToContext, nullptr, StopAtContext);
+			}
+		}
+	}
+
+	// Now add ours onto the tail
+	Fragments.Emplace(MoveTemp(NewLocator));
+	return true;
+}
+
+bool FUniversalObjectLocator::SerializeFromMismatchedTag(const FPropertyTag& Tag, FStructuredArchive::FSlot Slot)
+{
+	static const FName NAME_SoftObjectPath = "SoftObjectPath";
+
+	if (Tag.Type == NAME_SoftObjectProperty)
+	{
+		FSoftObjectPtr OldProperty;
+		Slot << OldProperty;
+
+		Fragments.Reset(1);
+		Fragments.Emplace(TUniversalObjectLocatorFragment<FDirectPathObjectLocator>(OldProperty.ToSoftObjectPath()));
+		return true;
+	}
+	else if (Tag.Type == NAME_StructProperty && Tag.StructName == NAME_SoftObjectPath)
+	{
+		FSoftObjectPath OldPath;
+		Slot << OldPath;
+
+		Fragments.Reset(1);
+		Fragments.Emplace(TUniversalObjectLocatorFragment<FDirectPathObjectLocator>(MoveTemp(OldPath)));
+		return true;
+	}
+
+	return false;
+}
+
+bool FUniversalObjectLocator::ExportTextItem(FString& ValueStr, const FUniversalObjectLocator& DefaultValue, UObject* Parent, int32 PortFlags, UObject* ExportRootScope) const
+{
+	TStringBuilder<128> String;
+	ToString(String);
+
+	ValueStr.AppendChar('(');
+	ValueStr.Append(String.ToString(), String.Len());
+	ValueStr.AppendChar(')');
+	return true;
+}
+
+bool FUniversalObjectLocator::ImportTextItem(const TCHAR*& Buffer, int32 PortFlags, UObject* Parent, FOutputDevice* ErrorText, FArchive* InSerializingArchive)
+{
+	using namespace UE::UniversalObjectLocator;
+
+	if (Buffer && *Buffer == '(')
+	{
+		const TCHAR* BufferEnd = FCString::Strchr(Buffer, ')');
+		if (Buffer != BufferEnd && (BufferEnd - Buffer) < std::numeric_limits<int32>::max())
+		{
+			FStringView View(Buffer + 1, int32(BufferEnd - Buffer) - 1);
+			if (TryParseString(View, FParseStringParams()))
+			{
+				return true;
+			}
+		}
+	}
+
+	// Try and parse this as a soft object path
+	FSoftObjectPath Path;
+	if (Path.ImportTextItem(Buffer, PortFlags, Parent, ErrorText, InSerializingArchive))
+	{
+		UObject* Object = Path.ResolveObject();
+		Reset(Object);
+		return true;
+	}
+
+	return false;
+}
+
+#undef LOCTEXT_NAMESPACE
\ No newline at end of file
diff --git a/Engine/Source/Runtime/UniversalObjectLocator/Private/UniversalObjectLocatorFragment.cpp b/Engine/Source/Runtime/UniversalObjectLocator/Private/UniversalObjectLocatorFragment.cpp
new file mode 100644
--- /dev/null
+++ b/Engine/Source/Runtime/UniversalObjectLocator/Private/UniversalObjectLocatorFragment.cpp
@@ -0,0 +1,676 @@
+// Copyright Epic Games, Inc. All Rights Reserved.
+
+#include "UniversalObjectLocatorFragment.h"
+#include "Algo/Find.h"
+#include "UniversalObjectLocatorFragmentType.h"
+#include "UniversalObjectLocatorStringParams.h"
+#include "UniversalObjectLocatorInitializeParams.h"
+#include "UniversalObjectLocatorInitializeResult.h"
+#include "UObject/SoftObjectPath.h"
+#include "Containers/SparseArray.h"
+#include "Misc/AsciiSet.h"
+
+#define LOCTEXT_NAMESPACE "UOL"
+
+DECLARE_LOG_CATEGORY_EXTERN(LogUniversalObjectLocator, Log, Warning);
+DEFINE_LOG_CATEGORY(LogUniversalObjectLocator);
+
+namespace UE::UniversalObjectLocator
+{
+	extern TArray<FFragmentType> GFragmentTypes;
+
+	const FFragmentType* FindBestFragmentType(const UObject* Object, const UObject* Context)
+	{
+		// Loop through all our FragmentTypes to find the most supported one
+		uint32 BestFragmentTypePriority = 0;
+		const FFragmentType* BestFragmentType = nullptr;
+
+		for (const FFragmentType& FragmentType : GFragmentTypes)
+		{
+			const uint32 ThisFragmentTypePriority = FragmentType.ComputePriority(Object, Context);
+			if (ThisFragmentTypePriority > BestFragmentTypePriority)
+			{
+				BestFragmentTypePriority = ThisFragmentTypePriority;
+				BestFragmentType = &FragmentType;
+			}
+		}
+
+		if (BestFragmentType && BestFragmentType->PayloadType != nullptr)
+		{
+			return BestFragmentType;
+		}
+
+		return nullptr;
+	}
+
+	FFragmentType* FFragmentTypeHandle::Resolve() const
+	{
+		return Handle == 0xff ? nullptr : &GFragmentTypes[Handle];
+	}
+
+	FFragmentTypeHandle MakeFragmentTypeHandle(const FFragmentType* FragmentType)
+	{
+		check(FragmentType);
+
+		const uint64 FragmentTypeOffset = static_cast<uint64>(FragmentType - GFragmentTypes.GetData());
+		checkf(FragmentTypeOffset < std::numeric_limits<uint8>::max(), TEXT("Maximum number of UOL FragmentTypes reached"));
+
+		return FFragmentTypeHandle(static_cast<uint8>(FragmentTypeOffset));
+	}
+
+} // UE::UniversalObjectLocator
+
+
+FUniversalObjectLocatorFragment::FUniversalObjectLocatorFragment(const UObject* InObject, const UObject* Context)
+	: bIsInitialized(0)
+	, bIsInline(0)
+{
+	Reset(InObject, Context);
+}
+
+FUniversalObjectLocatorFragment::FUniversalObjectLocatorFragment(const UE::UniversalObjectLocator::FFragmentType& InFragmentType)
+	: FragmentType(MakeFragmentTypeHandle(&InFragmentType))
+	, bIsInitialized(0)
+	, bIsInline(0)
+{
+	this->InitializePayload(InFragmentType.PayloadType);
+}
+
+FUniversalObjectLocatorFragment::FUniversalObjectLocatorFragment()
+	: bIsInitialized(0)
+	, bIsInline(0)
+{
+	static_assert(sizeof(FUniversalObjectLocatorFragment) == FUniversalObjectLocatorFragment::Size, "Unexpected size for FUniversalObjectLocatorFragment");
+	static_assert(offsetof(FUniversalObjectLocatorFragment, Data) == 0, "FUniversalObjectLocatorFragment inline data is not aligned properly");
+}
+
+FUniversalObjectLocatorFragment::~FUniversalObjectLocatorFragment()
+{
+	if (bIsInitialized && !IsEngineExitRequested())
+	{
+		DestroyPayload();
+	}
+}
+
+FUniversalObjectLocatorFragment::FUniversalObjectLocatorFragment(const FUniversalObjectLocatorFragment& RHS)
+	: FragmentType(RHS.FragmentType)
+	, bIsInitialized(0)
+	, bIsInline(0)
+{
+	using namespace UE::UniversalObjectLocator;
+
+	if (RHS.bIsInitialized)
+	{
+		const UScriptStruct* FragmentStruct = GetFragmentStruct();
+		check(FragmentStruct);
+
+		this->InitializePayload(FragmentStruct);
+
+		FragmentStruct->CopyScriptStruct(this->GetPayload(), RHS.GetPayload());
+	}
+	else
+	{
+		this->bIsInitialized = false;
+	}
+}
+
+FUniversalObjectLocatorFragment& FUniversalObjectLocatorFragment::operator=(const FUniversalObjectLocatorFragment& RHS)
+{
+	using namespace UE::UniversalObjectLocator;
+
+	this->DestroyPayload();
+
+	if (RHS.bIsInitialized)
+	{
+		const UScriptStruct* FragmentStruct = GetFragmentStruct();
+
+		// Assign the FragmentType and copy the payload
+		this->FragmentType = RHS.FragmentType;
+		this->InitializePayload(FragmentStruct);
+		FragmentStruct->CopyScriptStruct(this->GetPayload(), RHS.GetPayload());
+	}
+	else
+	{
+		this->bIsInitialized = false;
+		this->FragmentType   = FFragmentTypeHandle();
+	}
+
+	return *this;
+}
+
+FUniversalObjectLocatorFragment::FUniversalObjectLocatorFragment(FUniversalObjectLocatorFragment&& RHS)
+	: FragmentType(RHS.FragmentType)
+	, bIsInitialized(RHS.bIsInitialized)
+	, bIsInline(RHS.bIsInline)
+{
+	using namespace UE::UniversalObjectLocator;
+
+	FMemory::Memcpy(this->Data, RHS.Data, sizeof(Data));
+
+	RHS.bIsInitialized = false;
+	RHS.bIsInline      = false;
+	RHS.FragmentType   = FFragmentTypeHandle();
+}
+
+FUniversalObjectLocatorFragment& FUniversalObjectLocatorFragment::operator=(FUniversalObjectLocatorFragment&& RHS)
+{
+	using namespace UE::UniversalObjectLocator;
+
+	this->DestroyPayload();
+
+	this->bIsInitialized = RHS.bIsInitialized;
+	this->bIsInline      = RHS.bIsInline;
+	this->FragmentType   = RHS.FragmentType;
+	FMemory::Memcpy(this->Data, RHS.Data, sizeof(Data));
+
+	RHS.bIsInitialized = false;
+	RHS.bIsInline      = false;
+	RHS.FragmentType   = FFragmentTypeHandle();
+
+	return *this;
+}
+
+bool operator==(const FUniversalObjectLocatorFragment& A, const FUniversalObjectLocatorFragment& B)
+{
+	using namespace UE::UniversalObjectLocator;
+
+	if (A.bIsInitialized != B.bIsInitialized)
+	{
+		return false;
+	}
+	else if (!A.bIsInitialized)
+	{
+		// 2 uninitialized references are the same
+		return true;
+	}
+	else if (A.FragmentType != B.FragmentType)
+	{
+		// Different fragment types
+		return false;
+	}
+	else
+	{
+		const UScriptStruct* FragmentStruct = A.GetFragmentStruct();
+		check(FragmentStruct);
+
+		// Same fragment types - compare payloads
+		const void* PayloadA = A.GetPayload();
+		const void* PayloadB = B.GetPayload();
+
+		return FragmentStruct->CompareScriptStruct(PayloadA, PayloadB, 0);
+	}
+}
+
+bool operator!=(const FUniversalObjectLocatorFragment& A, const FUniversalObjectLocatorFragment& B)
+{
+	return !(A == B);
+}
+
+uint32 GetTypeHash(const FUniversalObjectLocatorFragment& Fragment)
+{
+	using namespace UE::UniversalObjectLocator;
+
+	if (!Fragment.bIsInitialized)
+	{
+		return 0;
+	}
+
+	uint32 Hash = GetTypeHash(Fragment.FragmentType);
+
+	if (const FFragmentType* FragmentTypePtr = Fragment.GetFragmentType())
+	{
+		UScriptStruct* Struct = FragmentTypePtr->GetStruct();
+		if (Struct)
+		{
+			const uint32 PayloadHash = Struct->GetStructTypeHash(Fragment.GetPayload());
+			Hash = HashCombineFast(Hash, PayloadHash);
+		}
+	}
+
+	return Hash;
+}
+
+#if DO_CHECK
+void FUniversalObjectLocatorFragment::CheckPayloadType(UScriptStruct* TypeToCompare) const
+{
+	using namespace UE::UniversalObjectLocator;
+
+	const FFragmentType* FragmentTypePtr = GetFragmentType();
+	checkf(FragmentTypePtr == nullptr || FragmentTypePtr->PayloadType == TypeToCompare,
+		TEXT("Type mismatch when accessing payload data! Attempting to access a stored %s payload as %s."),
+		FragmentTypePtr->PayloadType ? *FragmentTypePtr->PayloadType->GetName() : TEXT("<expired>"),
+		TypeToCompare          ? *TypeToCompare->GetName()          : TEXT("<nullptr>"));
+}
+#endif
+
+void FUniversalObjectLocatorFragment::ToString(FStringBuilderBase& OutString) const
+{
+	using namespace UE::UniversalObjectLocator;
+
+	const FFragmentType* FragmentTypePtr = FragmentType.Resolve();
+	if (FragmentTypePtr && FragmentTypePtr->PayloadType)
+	{
+		FragmentTypePtr->FragmentTypeID.AppendString(OutString);
+
+		TStringBuilder<128> PayloadString;
+		FragmentTypePtr->ToString(GetPayload(), PayloadString);
+		if (PayloadString.Len() != 0)
+		{
+			OutString += '=';
+			OutString.Append(PayloadString.ToView());
+		}
+	}
+}
+
+UE::UniversalObjectLocator::FParseStringResult FUniversalObjectLocatorFragment::TryParseString(FStringView InString, const FParseStringParams& InParams)
+{
+	using namespace UE::UniversalObjectLocator;
+
+	if (InString.Len() == 0)
+	{
+		Reset();
+		return FParseStringResult().Success();
+	}
+
+	// Check for a literal "none" text
+	static constexpr FStringView NoneString = TEXTVIEW("none");
+	if (InString.Compare(NoneString, ESearchCase::IgnoreCase) == 0)
+	{
+		Reset();
+		return FParseStringResult().Success(NoneString.Len());
+	}
+
+	FStringView FragmentTypeString = InString;
+	FStringView FragmentPayloadString;
+
+	const int32 Delimiter = UE::String::FindFirstChar(InString, '=');
+	if (Delimiter != INDEX_NONE)
+	{
+		// We have a payload
+		FragmentTypeString    = InString.Left(Delimiter);
+		FragmentPayloadString = InString.RightChop(Delimiter + 1);
+
+		if (FragmentTypeString.Len() == 0)
+		{
+			return FParseStringResult().Failure(UE_UOL_PARSE_ERROR(InParams, LOCTEXT("Error_UnexpectedEquals", "Unexpected '='' when expecting a fragment type.")));
+		}
+	}
+
+	FParseStringResult TypeResult = TryParseFragmentType(FragmentTypeString, InParams);
+	if (!TypeResult)
+	{
+		return TypeResult;
+	}
+
+	FParseStringResult PayloadResult = TryParseFragmentPayload(FragmentPayloadString, InParams);
+
+	// Add the type chars and = to the total num parsed
+	PayloadResult.NumCharsParsed += TypeResult.NumCharsParsed + 1;
+	return PayloadResult;
+}
+
+UE::UniversalObjectLocator::FParseStringResult FUniversalObjectLocatorFragment::TryParseFragmentType(FStringView InString, const FParseStringParams& InParams)
+{
+	using namespace UE::UniversalObjectLocator;
+
+	if (InString.Len() == 0)
+	{
+		return FParseStringResult().Failure(UE_UOL_PARSE_ERROR(InParams, LOCTEXT("Error_EmptyFragmentType", "Fragment type specifier is empty.")));
+	}
+
+	// Check for a literal "none" text
+	static constexpr FStringView NoneString = TEXTVIEW("none");
+	if (InString.Compare(NoneString, ESearchCase::IgnoreCase) == 0)
+	{
+		Reset();
+		return FParseStringResult().Success(NoneString.Len());
+	}
+
+	// Try and find the FragmentType as a name
+	FName FragmentTypeID(InString.Len(), InString.GetData(), FNAME_Find);
+	if (FragmentTypeID != NAME_None)
+	{
+		// Find the FragmentType
+		const FFragmentType* SerializedFragmentType = Algo::FindBy(GFragmentTypes, FragmentTypeID, &FFragmentType::FragmentTypeID);
+		if (SerializedFragmentType != nullptr && SerializedFragmentType->PayloadType != nullptr)
+		{
+			this->DestroyPayload();
+			this->FragmentType = MakeFragmentTypeHandle(SerializedFragmentType);
+			this->InitializePayload(SerializedFragmentType->PayloadType);
+
+			return FParseStringResult().Success(InString.Len());
+		}
+	}
+
+	// Not a valid fragment type string
+	return FParseStringResult().Failure(
+		UE_UOL_PARSE_ERROR(InParams, 
+			FText::Format(
+				LOCTEXT("Error_UnknownFragmentType", "Unknown fragment type specifier {0}."),
+				FText::FromStringView(InString)
+			)
+		)
+	);
+}
+
+UE::UniversalObjectLocator::FParseStringResult FUniversalObjectLocatorFragment::TryParseFragmentPayload(FStringView InString, const FParseStringParams& InParams)
+{
+	using namespace UE::UniversalObjectLocator;
+
+	if (!bIsInitialized)
+	{
+		return FParseStringResult().Failure(UE_UOL_PARSE_ERROR(InParams, LOCTEXT("Error_Uninitialized", "Unable to parse a payload for an uninitialized fragment.")));
+	}
+
+	const FFragmentType* FragmentTypePtr = GetFragmentType();
+	const UScriptStruct* FragmentStruct  = FragmentTypePtr ? FragmentTypePtr->GetStruct() : nullptr;
+
+	if (!FragmentStruct)
+	{
+		return FParseStringResult().Failure(UE_UOL_PARSE_ERROR(InParams, LOCTEXT("Error_Expired", "Unable to parse a payload for a fragment whose type has expired.")));
+	}
+
+	void* Payload = GetPayload();
+	if (InString.Len() == 0)
+	{
+		// Empty payload string means a default payload
+		FragmentStruct->ClearScriptStruct(Payload);
+		return FParseStringResult().Success();
+	}
+
+	return FragmentTypePtr->TryParseString(Payload, InString, InParams);
+}
+
+const UE::UniversalObjectLocator::FFragmentType* FUniversalObjectLocatorFragment::GetFragmentType() const
+{
+	return FragmentType.Resolve();
+}
+
+UScriptStruct* FUniversalObjectLocatorFragment::GetFragmentStruct() const
+{
+	using namespace UE::UniversalObjectLocator;
+	if (const FFragmentType* Type = GetFragmentType())
+	{
+		return Type->GetStruct();
+	}
+	return nullptr;
+}
+
+UE::UniversalObjectLocator::FFragmentTypeHandle FUniversalObjectLocatorFragment::GetFragmentTypeHandle() const
+{
+	using namespace UE::UniversalObjectLocator;
+
+	const FFragmentType* FragmentTypePtr = GetFragmentType();
+	if (FragmentTypePtr)
+	{
+		return MakeFragmentTypeHandle(FragmentTypePtr);
+	}
+	return FFragmentTypeHandle();
+}
+
+void FUniversalObjectLocatorFragment::DestroyPayload()
+{
+	using namespace UE::UniversalObjectLocator;
+
+	if (!bIsInitialized)
+	{
+		return;
+	}
+
+	void* Payload = GetPayload();
+
+	const UScriptStruct* FragmentStruct = GetFragmentStruct();
+	if (ensureMsgf(FragmentStruct, TEXT("FUniversalObjectLocatorFragment has outlived its FragmentType's payload type struct! This could leak memory if the type allocated it.")))
+	{
+		FragmentStruct->DestroyStruct(Payload);
+	}
+
+	if (!bIsInline)
+	{
+		FMemory::Free(Payload);
+	}
+
+	bIsInitialized = false;
+}
+
+void* FUniversalObjectLocatorFragment::GetPayload()
+{
+	check(bIsInitialized);
+	return bIsInline ? Data : *((void**)Data);
+}
+
+const void* FUniversalObjectLocatorFragment::GetPayload() const
+{
+	return bIsInline ? Data : *((const void* const *)Data);
+}
+
+void FUniversalObjectLocatorFragment::InitializePayload(const UScriptStruct* PayloadType)
+{
+	check(!bIsInitialized);
+
+	bIsInitialized = true;
+	if (PayloadType->GetStructureSize() <= sizeof(FUniversalObjectLocatorFragment::Data) && PayloadType->GetMinAlignment() <= alignof(FUniversalObjectLocatorFragment))
+	{
+		// We can placement new this into the payload data
+		bIsInline = true;
+	}
+	else
+	{
+		// We have to allocate this struct on the heap
+		void* HeapAllocation = FMemory::Malloc(PayloadType->GetStructureSize(), PayloadType->GetMinAlignment());
+		*reinterpret_cast<void**>(Data) = HeapAllocation;
+		bIsInline = false;
+	}
+
+	PayloadType->InitializeStruct(GetPayload());
+}
+
+void FUniversalObjectLocatorFragment::Reset()
+{
+	using namespace UE::UniversalObjectLocator;
+
+	DestroyPayload();
+	FragmentType = FFragmentTypeHandle();
+}
+
+void FUniversalObjectLocatorFragment::Reset(const UObject* InObject, const UObject* Context)
+{
+	using namespace UE::UniversalObjectLocator;
+
+	Reset();
+
+	if (const FFragmentType* BestFragmentType = FindBestFragmentType(InObject, Context))
+	{
+		FragmentType = MakeFragmentTypeHandle(BestFragmentType);
+		InitializePayload(BestFragmentType->PayloadType);
+
+		BestFragmentType->InitializePayload(GetPayload(), FInitializeParams(InObject, Context));
+	}
+}
+
+void FUniversalObjectLocatorFragment::Reset(const UObject* InObject, const UObject* Context, TFunctionRef<bool(UE::UniversalObjectLocator::FFragmentTypeHandle)> CanUseFragmentType)
+{
+	using namespace UE::UniversalObjectLocator;
+
+	Reset();
+
+	// Loop through all our FragmentTypes to find the most supported one
+	uint32 BestFragmentTypePriority = 0;
+	const FFragmentType* BestFragmentType = nullptr;
+
+	if (!ensure(GFragmentTypes.Num() < 255))
+	{
+		return;
+	}
+
+	const uint8 Num = static_cast<uint8>(GFragmentTypes.Num());
+	for (uint8 Index = 0; Index < Num; ++Index)
+	{
+		if (!CanUseFragmentType(FFragmentTypeHandle(Index)))
+		{
+			continue;
+		}
+
+		const FFragmentType& ThisFragmentType = GFragmentTypes[Index];
+		const uint32 ThisFragmentTypePriority = ThisFragmentType.ComputePriority(InObject, Context);
+		if (ThisFragmentTypePriority > BestFragmentTypePriority)
+		{
+			BestFragmentTypePriority = ThisFragmentTypePriority;
+			BestFragmentType = &ThisFragmentType;
+		}
+	}
+
+	if (BestFragmentType && BestFragmentType->PayloadType != nullptr)
+	{
+		FragmentType = MakeFragmentTypeHandle(BestFragmentType);
+		InitializePayload(BestFragmentType->PayloadType);
+
+		BestFragmentType->InitializePayload(GetPayload(), FInitializeParams(InObject, Context));
+	}
+}
+
+UE::UniversalObjectLocator::FResolveResult FUniversalObjectLocatorFragment::Resolve(const UE::UniversalObjectLocator::FResolveParams& Params) const
+{
+	using namespace UE::UniversalObjectLocator;
+
+	// Find our FragmentType
+	const FFragmentType* FragmentTypePtr = FragmentType.Resolve();
+	if (FragmentTypePtr && bIsInitialized)
+	{
+		return FragmentTypePtr->ResolvePayload(GetPayload(), Params);
+	}
+
+	return FResolveResult();
+}
+
+bool FUniversalObjectLocatorFragment::Serialize(FArchive& Ar)
+{
+	using namespace UE::UniversalObjectLocator;
+
+	if (Ar.IsLoading())
+	{
+		FName FragmentTypeID;
+		Ar << FragmentTypeID;
+
+		if (FragmentTypeID == NAME_None)
+		{
+			Reset();
+		}
+		else
+		{
+			// Find the FragmentType
+			const FFragmentType* SerializedFragmentType = Algo::FindBy(GFragmentTypes, FragmentTypeID, &FFragmentType::FragmentTypeID);
+			if (!SerializedFragmentType || !SerializedFragmentType->PayloadType)
+			{
+				Reset();
+
+				// Big error - what do we do?
+				UE_LOG(LogUniversalObjectLocator, Error, TEXT("WARNING: POTENTIAL DATA LOSS! Universal Object Reference FragmentType %s is not recognized! This reference will be lost if re-saved."), *FragmentTypeID.ToString());
+				checkf(false, TEXT("WARNING: POTENTIAL DATA LOSS! Universal Object Reference FragmentType %s! This reference will be lost if re-saved."), *FragmentTypeID.ToString());
+
+				// Deserialize an empty payload so we don't corrupt the serialization data
+				FUniversalObjectLocatorEmptyPayload Empty;
+				FUniversalObjectLocatorEmptyPayload::StaticStruct()->SerializeItem(Ar, &Empty, nullptr);
+				return true;
+			}
+
+			FragmentType = MakeFragmentTypeHandle(SerializedFragmentType);
+
+			UScriptStruct* Struct = SerializedFragmentType->GetStruct();
+			Ar.Preload(Struct);
+
+			InitializePayload(Struct);
+			Struct->SerializeItem(Ar, GetPayload(), nullptr);
+		}
+	}
+	else if (Ar.IsSaving() || Ar.IsTransacting())
+	{
+		const FFragmentType* FragmentTypePtr = FragmentType.Resolve();
+
+		if (FragmentTypePtr == nullptr)
+		{
+			FName None;
+
+			// FragmentType ID
+			Ar << None;
+		}
+		else
+		{
+			FName FragmentTypeID = FragmentTypePtr->FragmentTypeID;
+
+			// FragmentType ID
+			Ar << FragmentTypeID;
+
+			// FragmentType payload
+			FragmentTypePtr->GetStruct()->SerializeItem(Ar, GetPayload(), nullptr);
+		}
+	}
+	else if (Ar.IsModifyingWeakAndStrongReferences())
+	{
+		UScriptStruct* FragmentStruct = GetFragmentStruct();
+		if (FragmentStruct && bIsInitialized)
+		{
+			FragmentStruct->SerializeItem(Ar, GetPayload(), nullptr);
+		}
+	}
+
+	return true;
+}
+
+void FUniversalObjectLocatorFragment::AddStructReferencedObjects(FReferenceCollector& Collector)
+{
+	using namespace UE::UniversalObjectLocator;
+
+	if (FFragmentType* FragmentTypePtr = FragmentType.Resolve())
+	{
+		Collector.AddReferencedObject(FragmentTypePtr->PayloadType);
+		if (bIsInitialized && FragmentTypePtr->PayloadType)
+		{
+			Collector.AddReferencedObjects(FragmentTypePtr->PayloadType, GetPayload());
+		}
+	}
+}
+
+bool FUniversalObjectLocatorFragment::ExportTextItem(FString& ValueStr, const FUniversalObjectLocatorFragment& DefaultValue, UObject* Parent, int32 PortFlags, UObject* ExportRootScope) const
+{
+	TStringBuilder<32> PayloadString;
+	ToString(PayloadString);
+
+	ValueStr.AppendChar('(');
+	ValueStr.Append(PayloadString.ToString(), PayloadString.Len());
+	ValueStr.AppendChar(')');
+
+	return true;
+}
+
+bool FUniversalObjectLocatorFragment::ImportTextItem(const TCHAR*& Buffer, int32 PortFlags, UObject* Parent, FOutputDevice* ErrorText, FArchive* InSerializingArchive)
+{
+	using namespace UE::UniversalObjectLocator;
+
+	if (Buffer && *Buffer == '(')
+	{
+		const TCHAR* BufferEnd = FCString::Strchr(Buffer, ')');
+		if (Buffer != BufferEnd && (BufferEnd - Buffer) < std::numeric_limits<int32>::max())
+		{
+			FStringView View(Buffer+1, int32(BufferEnd-Buffer)-1);
+			if (TryParseString(View, FParseStringParams()))
+			{
+				return true;
+			}
+		}
+	}
+
+	return false;
+}
+
+bool FUniversalObjectLocatorFragment::SerializeFromMismatchedTag(const FPropertyTag& Tag, FStructuredArchive::FSlot Slot)
+{
+	return false;
+}
+
+void FUniversalObjectLocatorFragment::GetPreloadDependencies(TArray<UObject*>& OutDeps)
+{
+
+}
+
+#undef LOCTEXT_NAMESPACE
\ No newline at end of file
diff --git a/Engine/Source/Runtime/UniversalObjectLocator/Private/UniversalObjectLocatorFragmentType.cpp b/Engine/Source/Runtime/UniversalObjectLocator/Private/UniversalObjectLocatorFragmentType.cpp
new file mode 100644
--- /dev/null
+++ b/Engine/Source/Runtime/UniversalObjectLocator/Private/UniversalObjectLocatorFragmentType.cpp
@@ -0,0 +1,41 @@
+// Copyright Epic Games, Inc. All Rights Reserved.
+
+#include "UniversalObjectLocatorFragmentType.h"
+#include "UniversalObjectLocatorFragment.h"
+#include "UniversalObjectLocatorStringParams.h"
+#include "UniversalObjectLocatorInitializeParams.h"
+#include "UniversalObjectLocatorInitializeResult.h"
+
+namespace UE::UniversalObjectLocator
+{
+
+FResolveResult FFragmentType::ResolvePayload(const void* Payload, const FResolveParams& Params) const
+{
+	return (*InstanceBindings.Resolve)(Payload, Params);
+}
+FInitializeResult FFragmentType::InitializePayload(void* Payload, const FInitializeParams& InParams) const
+{
+	return (*InstanceBindings.Initialize)(Payload, InParams);
+}
+void FFragmentType::ToString(const void* Payload, FStringBuilderBase& OutStringBuilder) const
+{
+	const int32 StartPos = OutStringBuilder.Len();
+
+	(*InstanceBindings.ToString)(Payload, OutStringBuilder);
+
+#if DO_CHECK
+	FStringView StringRepresentation(OutStringBuilder.GetData() + StartPos, OutStringBuilder.Len());
+	checkf(!FAsciiSet::HasAny(StringRepresentation, ~FUniversalObjectLocatorFragment::ValidFragmentPayloadCharacters), TEXT("F%s::ToString resulted in an invalid character usage"), *PayloadType->GetName());
+#endif
+}
+FParseStringResult FFragmentType::TryParseString(void* Payload, FStringView InString, const FParseStringParams& Params) const
+{
+	return (*InstanceBindings.TryParseString)(Payload, InString, Params);
+}
+
+uint32 FFragmentType::ComputePriority(const UObject* Object, const UObject* Context) const
+{
+	return (*StaticBindings.Priority)(Object, Context);
+}
+
+} // namespace UE::UniversalObjectLocator
\ No newline at end of file
diff --git a/Engine/Source/Runtime/UniversalObjectLocator/Private/UniversalObjectLocatorModule.cpp b/Engine/Source/Runtime/UniversalObjectLocator/Private/UniversalObjectLocatorModule.cpp
new file mode 100644
--- /dev/null
+++ b/Engine/Source/Runtime/UniversalObjectLocator/Private/UniversalObjectLocatorModule.cpp
@@ -0,0 +1,80 @@
+// Copyright Epic Games, Inc. All Rights Reserved.
+
+#include "SubObjectLocator.h"
+#include "Modules/ModuleManager.h"
+#include "IUniversalObjectLocatorModule.h"
+#include "UniversalObjectLocatorFragment.h"
+//#include "Modules/VisualizerDebuggingState.h"
+#include "Containers/Ticker.h"
+
+#include "DirectPathObjectLocator.h"
+
+namespace UE::UniversalObjectLocator
+{
+
+TArray<FFragmentType> GFragmentTypes;
+
+class FUniversalObjectLocatorModule
+	: public IUniversalObjectLocatorModule
+{
+public:
+	
+	void StartupModule() override
+	{
+		{
+			FFragmentTypeParameters FragmentTypeParams("subobj", NSLOCTEXT("SubObjectLocator", "Object", "Object"));
+			FragmentTypeParams.PrimaryEditorType = "SubObject";
+			FSubObjectLocator::FragmentType = RegisterFragmentType<FSubObjectLocator>(FragmentTypeParams);
+		}
+
+		{
+			FFragmentTypeParameters FragmentTypeParams("uobj", NSLOCTEXT("DirectPathObjectLocator", "Object", "Object"));
+			FDirectPathObjectLocator::FragmentType = RegisterFragmentType<FDirectPathObjectLocator>(FragmentTypeParams);
+		}
+
+		TickerDelegate = FTSTicker::GetCoreTicker().AddTicker(
+			FTickerDelegate::CreateRaw(this, &FUniversalObjectLocatorModule::PurgeVisualizers), 60.f);
+	}
+
+	void ShutdownModule() override
+	{
+
+	}
+
+	FFragmentTypeHandle RegisterFragmentTypeImpl(const FFragmentType& FragmentType) override
+	{
+		const int32 Index = GFragmentTypes.Num();
+		GFragmentTypes.Add(FragmentType);
+		checkf(Index < static_cast<int32>(std::numeric_limits<uint8>::max()), TEXT("Maximum number of UOL FragmentTypes reached"));
+
+		// @todo: enable this code once visualizer debugging state is enabled
+		// Re-assign the debugging ptr in case it changed
+		// UE::Core::FVisualizerDebuggingState::Assign("UOL", GFragmentTypes.GetData());
+
+		return FFragmentTypeHandle(static_cast<uint8>(Index));
+	}
+
+	void UnregisterFragmentTypeImpl(FFragmentTypeHandle FragmentType) override
+	{
+		GFragmentTypes[FragmentType.GetIndex()] = FFragmentType{};
+	}
+
+	bool PurgeVisualizers(float) const
+	{
+		for (FFragmentType& FragmentType : GFragmentTypes)
+		{
+			if (FragmentType.DebuggingAssistant)
+			{
+				FragmentType.DebuggingAssistant->Purge();
+			}
+		}
+		return true;
+	}
+
+	FTSTicker::FDelegateHandle TickerDelegate;
+};
+
+} // namespace UE::UniversalObjectLocator
+
+IMPLEMENT_MODULE(UE::UniversalObjectLocator::FUniversalObjectLocatorModule, UniversalObjectLocator);
+
diff --git a/Engine/Source/Runtime/UniversalObjectLocator/Private/UniversalObjectLocatorScriptingExtensions.cpp b/Engine/Source/Runtime/UniversalObjectLocator/Private/UniversalObjectLocatorScriptingExtensions.cpp
new file mode 100644
--- /dev/null
+++ b/Engine/Source/Runtime/UniversalObjectLocator/Private/UniversalObjectLocatorScriptingExtensions.cpp
@@ -0,0 +1,60 @@
+// Copyright Epic Games, Inc. All Rights Reserved.
+
+#include "UniversalObjectLocatorScriptingExtensions.h"
+#include "UniversalObjectLocator.h"
+#include "UniversalObjectLocatorStringParams.h"
+#include "UObject/Stack.h"
+
+FUniversalObjectLocator UUniversalObjectLocatorScriptingExtensions::MakeUniversalObjectLocator(UObject* Object, const UObject* Context)
+{
+	return FUniversalObjectLocator(Object, Context);
+}
+
+FUniversalObjectLocator UUniversalObjectLocatorScriptingExtensions::UniversalObjectLocatorFromString(const FString& InString)
+{
+	using namespace::UE::UniversalObjectLocator;
+
+	FParseStringParams Params;
+	Params.Flags = EParseStringFlags::ErrorMessaging;
+
+	FUniversalObjectLocator Locator;
+	FParseStringResult Result = Locator.TryParseString(InString, Params);
+	if (!Result)
+	{
+		FFrame::KismetExecutionMessage(
+			*FText::Format(
+				NSLOCTEXT("UOL", "ParseError", "Unable to parse a valid UOL from string '{0}': {1}."),
+				FText::FromStringView(InString),
+				Result.ErrorMessage
+			).ToString()
+		, ELogVerbosity::Error);
+	}
+	return Locator;
+}
+
+bool UUniversalObjectLocatorScriptingExtensions::IsEmpty(const FUniversalObjectLocator& Locator)
+{
+	return Locator.IsEmpty();
+}
+
+FString UUniversalObjectLocatorScriptingExtensions::ToString(const FUniversalObjectLocator& Locator)
+{
+	TStringBuilder<128> String;
+	Locator.ToString(String);
+	return String.ToString();
+}
+
+UObject* UUniversalObjectLocatorScriptingExtensions::SyncFind(const FUniversalObjectLocator& Locator, UObject* Context)
+{
+	return Locator.SyncFind(Context);
+}
+
+UObject* UUniversalObjectLocatorScriptingExtensions::SyncLoad(const FUniversalObjectLocator& Locator, UObject* Context)
+{
+	return Locator.SyncLoad(Context);
+}
+
+void UUniversalObjectLocatorScriptingExtensions::SyncUnload(const FUniversalObjectLocator& Locator, UObject* Context)
+{
+	Locator.SyncUnload(Context);
+}
diff --git a/Engine/Source/Runtime/UniversalObjectLocator/Private/UniversalObjectLocatorStringParams.cpp b/Engine/Source/Runtime/UniversalObjectLocator/Private/UniversalObjectLocatorStringParams.cpp
new file mode 100644
--- /dev/null
+++ b/Engine/Source/Runtime/UniversalObjectLocator/Private/UniversalObjectLocatorStringParams.cpp
@@ -0,0 +1,17 @@
+// Copyright Epic Games, Inc. All Rights Reserved.
+
+#include "UniversalObjectLocatorStringParams.h"
+#include "Containers/StringView.h"
+
+namespace UE::UniversalObjectLocator
+{
+
+FStringView FParseStringResult::Progress(FStringView CurrentString, int32 NumToProgress)
+{
+	checkf(NumToProgress <= CurrentString.Len(), TEXT("Attempting to progress past the end of the current string!"));
+
+	NumCharsParsed += NumToProgress;
+	return CurrentString.RightChop(NumToProgress);
+}
+
+} // namespace UE::UniversalObjectLocator
\ No newline at end of file
diff --git a/Engine/Source/Runtime/UniversalObjectLocator/Private/UniversalObjectLocatorStringUtils.cpp b/Engine/Source/Runtime/UniversalObjectLocator/Private/UniversalObjectLocatorStringUtils.cpp
new file mode 100644
--- /dev/null
+++ b/Engine/Source/Runtime/UniversalObjectLocator/Private/UniversalObjectLocatorStringUtils.cpp
@@ -0,0 +1,67 @@
+// Copyright Epic Games, Inc. All Rights Reserved.
+
+#include "UniversalObjectLocatorStringUtils.h"
+#include "Containers/StringView.h"
+#include "Misc/Char.h"
+
+namespace UE::UniversalObjectLocator
+{
+
+bool ParseUnsignedInteger(FStringView InString, uint32& OutResult, int32* OutNumCharsParsed)
+{
+	if (InString.Len() == 0)
+	{
+		return false;
+	}
+
+	uint32 Parsed = 0;
+	uint32 Factor = 1;
+	int32  Index  = 0;
+
+	for (; Index < InString.Len(); ++Index)
+	{
+		// Skip leading zeros
+		const TCHAR Current = InString[Index];
+		if (!FChar::IsDigit(Current))
+		{
+			break;
+		}
+
+		const uint32 ThisNumber = (Current - '0');
+
+		// Skip leading zeros
+		if (Parsed == 0 && ThisNumber == 0)
+		{
+			continue;
+		}
+
+		auto IsValidOperation = [](uint32 Start, uint32 Now) { return Now >= Start && Now < static_cast<uint32>(std::numeric_limits<int32>::max()); };
+
+		// Multiply by 10 using bit-shift trick (x*10 === x*2 + x*8 === x<<1 + x<<3) and check for overflow
+		const uint32 Value   = ThisNumber * Factor;
+		const uint32 Shifted = Parsed * 10;
+
+		Factor *= 10;
+
+		if (Shifted < Parsed || Shifted+Value < Parsed)
+		{
+			// Overflow - can't be a uint
+			return false;
+		}
+
+		Parsed = Shifted+Value;
+	}
+
+	if (Index != 0)
+	{
+		OutResult = Parsed;
+		if (OutNumCharsParsed)
+		{
+			*OutNumCharsParsed = Index;
+		}
+	}
+	return Index != 0;
+}
+
+
+} // namespace UE::UniversalObjectLocator
\ No newline at end of file
diff --git a/Engine/Source/Runtime/UniversalObjectLocator/Public/DirectPathObjectLocator.h b/Engine/Source/Runtime/UniversalObjectLocator/Public/DirectPathObjectLocator.h
new file mode 100644
--- /dev/null
+++ b/Engine/Source/Runtime/UniversalObjectLocator/Public/DirectPathObjectLocator.h
@@ -0,0 +1,51 @@
+// Copyright Epic Games, Inc. All Rights Reserved.
+
+#pragma once
+
+#include "CoreTypes.h"
+#include "UObject/SoftObjectPath.h"
+#include "UniversalObjectLocatorFwd.h"
+#include "DirectPathObjectLocator.generated.h"
+
+class UObject;
+
+/**
+ * Object locator type that simply references its object by a direct path
+ */
+USTRUCT()
+struct FDirectPathObjectLocator
+{
+	GENERATED_BODY()
+
+	UPROPERTY()
+	FSoftObjectPath Path;
+
+	UNIVERSALOBJECTLOCATOR_API static UE::UniversalObjectLocator::TFragmentTypeHandle<FDirectPathObjectLocator> FragmentType;
+
+	friend uint32 GetTypeHash(const FDirectPathObjectLocator& A)
+	{
+		return GetTypeHash(A.Path);
+	}
+
+	friend bool operator==(const FDirectPathObjectLocator& A, const FDirectPathObjectLocator& B)
+	{
+		return A.Path == B.Path;
+	}
+
+	UE::UniversalObjectLocator::FResolveResult Resolve(const UE::UniversalObjectLocator::FResolveParams& Params) const;
+	UE::UniversalObjectLocator::FInitializeResult Initialize(const UE::UniversalObjectLocator::FInitializeParams& InParams);
+
+	void ToString(FStringBuilderBase& OutStringBuilder) const;
+	UE::UniversalObjectLocator::FParseStringResult TryParseString(FStringView InString, const UE::UniversalObjectLocator::FParseStringParams& Params);
+
+	static uint32 ComputePriority(const UObject* Object, const UObject* Context);
+};
+
+template<>
+struct TStructOpsTypeTraits<FDirectPathObjectLocator> : public TStructOpsTypeTraitsBase2<FDirectPathObjectLocator>
+{
+	enum
+	{
+		WithIdenticalViaEquality = true,
+	};
+};
diff --git a/Engine/Source/Runtime/UniversalObjectLocator/Public/IUniversalObjectLocatorModule.h b/Engine/Source/Runtime/UniversalObjectLocator/Public/IUniversalObjectLocatorModule.h
new file mode 100644
--- /dev/null
+++ b/Engine/Source/Runtime/UniversalObjectLocator/Public/IUniversalObjectLocatorModule.h
@@ -0,0 +1,114 @@
+// Copyright Epic Games, Inc. All Rights Reserved.
+
+#pragma once
+
+#include "Modules/ModuleInterface.h"
+#include "UniversalObjectLocatorFwd.h"
+#include "UniversalObjectLocatorFragmentTypeHandle.h"
+#include "UniversalObjectLocatorFragmentType.h"
+#include "UniversalObjectLocatorStringParams.h"
+#include "UniversalObjectLocatorInitializeResult.h"
+#include "UniversalObjectLocatorFragment.h"
+#include "Concepts/GetTypeHashable.h"
+#include "Templates/UnrealTypeTraits.h"
+#include "Misc/GeneratedTypeName.h"
+
+struct FSubObjectLocator;
+
+namespace UE::UniversalObjectLocator
+{
+
+struct FFragmentTypeParameters
+{
+	FFragmentTypeParameters(FName InFragmentTypeID, FText InDisplayText)
+		: DisplayText(InDisplayText)
+		, FragmentTypeID(InFragmentTypeID)
+	{
+	}
+
+	FText DisplayText;
+	FName FragmentTypeID;
+	FName PrimaryEditorType;
+};
+
+class IUniversalObjectLocatorModule
+	: public IModuleInterface
+{
+public:
+
+
+	/**
+	 * Register a new fragment type that can be used to locate objects from a Universal Object Locator
+	 *
+	 * Fragment types must be serializable USTRUCTS, must be equality comparable, and must be GetTypeHashable.
+	 * Additionally, the following member functions must be defined:
+	 *
+	 *   // Initialize the fragment to point at an object within a context
+	 *   FInitializeResult Initialize(const FInitializeParams& InParams) const;
+	 *   // Resolve the fragment to the resulting object, potentially loading or unloading it
+	 *   FResolveResult Resolve(const FResolveParams& Params) const;
+	 *
+	 *   // Convert this fragment to its string representation
+	 *   void ToString(FStringBuilderBase& OutStringBuilder) const;
+	 *   // Try and parse this fragment from its string representation
+	 *   FParseStringResult TryParseString(FStringView InString, const FParseStringParams& Params) const;
+	 *
+	 *   // Compute a priority for this fragment based on an object and context in order to decide which fragment type should be used if multiple can address the same object
+	 *   static uint32 ComputePriority(const UObject* Object, const UObject* Context) const;
+	 */
+	template<typename PayloadStructType
+		UE_REQUIRES(
+			// Payloads must be equality-comparable (we don't use the CEqualityComparable concept here because we need to make sure the ICppStructOps are set up correctly)
+			(TStructOpsTypeTraits<PayloadStructType>::WithIdenticalViaEquality || TStructOpsTypeTraits<PayloadStructType>::WithIdentical) &&
+			// Payloads must be type-hashable
+			TModels_V<CGetTypeHashable, PayloadStructType>
+		)
+	>
+	TFragmentTypeHandle<PayloadStructType> RegisterFragmentType(const FFragmentTypeParameters& FragmentTypeParameters)
+	{
+		checkf(FragmentTypeParameters.FragmentTypeID != NAME_None, TEXT("'None' is not a valid Fragment Type ID for type %s"), GetGeneratedTypeName<PayloadStructType>());
+		checkf(!FAsciiSet::HasAny(*FragmentTypeParameters.FragmentTypeID.ToString(), ~FUniversalObjectLocatorFragment::ValidFragmentTypeCharacters), TEXT("Fragment Type ID '%s' contains invalid characters"), *FragmentTypeParameters.FragmentTypeID.ToString());
+
+		FFragmentType NewFragmentType;
+		NewFragmentType.FragmentTypeID     = FragmentTypeParameters.FragmentTypeID;
+		NewFragmentType.PrimaryEditorType  = FragmentTypeParameters.PrimaryEditorType;
+		NewFragmentType.DebuggingAssistant = MakeShared<TFragmentTypeDebuggingAssistant<PayloadStructType>>();
+		NewFragmentType.PayloadType        = PayloadStructType::StaticStruct();
+
+		// Static bindings
+		NewFragmentType.StaticBindings.Priority = PayloadStructType::ComputePriority;
+
+		// Member function bindings
+		NewFragmentType.InstanceBindings.Resolve = [](const void* Payload, const FResolveParams& Params) -> FResolveResult
+		{
+			return static_cast<const PayloadStructType*>(Payload)->Resolve(Params);
+		};
+		NewFragmentType.InstanceBindings.Initialize = [](void* Payload, const FInitializeParams& InParams) -> FInitializeResult
+		{
+			return static_cast<PayloadStructType*>(Payload)->Initialize(InParams);
+		};
+		NewFragmentType.InstanceBindings.ToString = [](const void* Payload, FStringBuilderBase& OutStringBuilder)
+		{
+			return static_cast<const PayloadStructType*>(Payload)->ToString(OutStringBuilder);
+		};
+		NewFragmentType.InstanceBindings.TryParseString = [](void* Payload, FStringView InString, const FParseStringParams& Params) -> FParseStringResult
+		{
+			return static_cast<PayloadStructType*>(Payload)->TryParseString(InString, Params);
+		};
+		return TFragmentTypeHandle<PayloadStructType>(RegisterFragmentTypeImpl(NewFragmentType));
+	}
+
+	template<typename PayloadStructType>
+	void UnregisterFragmentType(TFragmentTypeHandle<PayloadStructType> FragmentType)
+	{
+		UnregisterFragmentTypeImpl(FragmentType);
+	}
+
+private:
+
+	virtual FFragmentTypeHandle RegisterFragmentTypeImpl(const FFragmentType& FragmentType) = 0;
+	virtual void UnregisterFragmentTypeImpl(FFragmentTypeHandle FragmentType) = 0;
+};
+
+
+} // namespace UE::UniversalObjectLocator
\ No newline at end of file
diff --git a/Engine/Source/Runtime/UniversalObjectLocator/Public/SubObjectLocator.h b/Engine/Source/Runtime/UniversalObjectLocator/Public/SubObjectLocator.h
new file mode 100644
--- /dev/null
+++ b/Engine/Source/Runtime/UniversalObjectLocator/Public/SubObjectLocator.h
@@ -0,0 +1,47 @@
+// Copyright Epic Games, Inc. All Rights Reserved.
+
+#pragma once
+
+#include "CoreTypes.h"
+#include "UObject/SoftObjectPath.h"
+#include "UniversalObjectLocatorFwd.h"
+#include "SubObjectLocator.generated.h"
+
+class UObject;
+
+USTRUCT()
+struct FSubObjectLocator
+{
+	GENERATED_BODY()
+
+	UPROPERTY()
+	FString PathWithinContext;
+
+	UNIVERSALOBJECTLOCATOR_API static UE::UniversalObjectLocator::TFragmentTypeHandle<FSubObjectLocator> FragmentType;
+
+	friend uint32 GetTypeHash(const FSubObjectLocator& A)
+	{
+		return GetTypeHash(A.PathWithinContext);
+	}
+
+	friend bool operator==(const FSubObjectLocator& A, const FSubObjectLocator& B)
+	{
+		return A.PathWithinContext == B.PathWithinContext;
+	}
+
+	UE::UniversalObjectLocator::FResolveResult Resolve(const UE::UniversalObjectLocator::FResolveParams& Params) const;
+	UE::UniversalObjectLocator::FInitializeResult Initialize(const UE::UniversalObjectLocator::FInitializeParams& InParams);
+	void ToString(FStringBuilderBase& OutStringBuilder) const;
+	UE::UniversalObjectLocator::FParseStringResult TryParseString(FStringView InString, const UE::UniversalObjectLocator::FParseStringParams& Params);
+
+	static uint32 ComputePriority(const UObject* Object, const UObject* Context);
+};
+
+template<>
+struct TStructOpsTypeTraits<FSubObjectLocator> : public TStructOpsTypeTraitsBase2<FSubObjectLocator>
+{
+	enum
+	{
+		WithIdenticalViaEquality = true,
+	};
+};
diff --git a/Engine/Source/Runtime/UniversalObjectLocator/Public/UniversalObjectLocator.h b/Engine/Source/Runtime/UniversalObjectLocator/Public/UniversalObjectLocator.h
new file mode 100644
--- /dev/null
+++ b/Engine/Source/Runtime/UniversalObjectLocator/Public/UniversalObjectLocator.h
@@ -0,0 +1,231 @@
+// Copyright Epic Games, Inc. All Rights Reserved.
+
+#pragma once
+
+#include "CoreTypes.h"
+#include "UniversalObjectLocatorFwd.h"
+#include "UniversalObjectLocatorFragment.h"
+#include "UniversalObjectLocatorResolveParams.h"
+#include "UniversalObjectLocator.generated.h"
+
+/**
+ * Universal Object Locators (UOLs) define an address to an object.
+ *
+ * The address is implemented as a chain of FUniversalObjectLocatorFragments, allowing addressing of objects
+ *   that may be nested deeply within levels of externally defined spawn or ownership logic.
+ *
+ * For example, a Universal Object Locator may reference an Anim Instance within a Skeletal Mesh Actor
+ *   is spawned by a Child Actor Component that is spawned by Sequencer. This is impossible with a
+ *   regular soft object path, but is perfectly feasible for a UOL.
+ *
+ * This type is 16 bytes.
+ */
+USTRUCT(BlueprintType, Category=GameFramework, meta=(HasNativeMake="/Script/UniversalObjectLocator.UniversalObjectLocatorScriptingExtensions.MakeUniversalObjectLocator"))
+struct FUniversalObjectLocator
+{
+	GENERATED_BODY()
+
+	using FResolveParams       = UE::UniversalObjectLocator::FResolveParams;
+	using FResolveResult       = UE::UniversalObjectLocator::FResolveResult;
+	using FResolveResultData   = UE::UniversalObjectLocator::FResolveResultData;
+	using FParseStringParams   = UE::UniversalObjectLocator::FParseStringParams;
+	using FParseStringResult   = UE::UniversalObjectLocator::FParseStringResult;
+
+	/**
+	 * Default constructor
+	 */
+	UNIVERSALOBJECTLOCATOR_API FUniversalObjectLocator();
+
+	/**
+	 * Attempt to construct this locator from a given object.
+	 * May result in an empty locator if no suitable address could be created.
+	 *
+	 * @param Object           The object that this locator should represent
+	 * @param Context          (Optional) Constrain this universal reference based on the specified context. This context should be passed to Resolve otherwise the resolution may fail.
+	 * @param StopAtContext    (Optional) Stop constructing this universal reference when we reach the following context (can be used if that context is always passed to Resolve to keep this type smaller)
+	 */
+	UNIVERSALOBJECTLOCATOR_API FUniversalObjectLocator(UObject* Object, const UObject* Context = nullptr, const UObject* StopAtContext = nullptr);
+
+	/**
+	 * Check if this locator is 'empty'. An empty locator contains no fragments and will never resolve.
+	 */
+	bool IsEmpty() const
+	{
+		return Fragments.Num() == 0;
+	}
+
+	/**
+	 * Attempt to resolve this locator by invoking the payload's 'Resolve' function
+	 *
+	 * @param Params           Resolution parameters defining the type of resolution to perform
+	 * @return A result structure defining the resolved object pointer, and associated flags
+	 */
+	UNIVERSALOBJECTLOCATOR_API FResolveResult Resolve(const FResolveParams& Params) const;
+
+	/**
+	 * Attempt to find the object this locator points to.
+	 * Shorthand for Resolve(FResolveParams::AsyncFind(Context)).AsObject().ReleaseFuture().
+	 *
+	 * @param Context          (Optional) An optional context to use for finding the object - should match what was specified in Reset() or in construction
+	 * @return A result structure that may or may not already have a result populated
+	 */
+	UNIVERSALOBJECTLOCATOR_API FResolveResult AsyncFind(const UObject* Context = nullptr) const;
+
+	/**
+	 * Attempt to find the object this locator points to, loading it if necessary (and possible).
+	 * Shorthand for Resolve(FResolveParams::AsyncLoad(Context)).AsObject().ReleaseFuture().
+	 *
+	 * @param Context          (Optional) An optional context to use for finding/loading the object - should match what was specified in Reset()
+	 * @return A result structure that may or may not already have a result populated
+	 */
+	UNIVERSALOBJECTLOCATOR_API FResolveResult AsyncLoad(const UObject* Context = nullptr) const;
+
+	/**
+	 * Attempt to unload the object this locator points to if possible.
+	 * Shorthand for Resolve(FResolveParams::AsyncUnload(Context)).AsVoid().ReleaseFuture().
+	 *
+	 * @param Context          (Optional) An optional context to use for finding/loading the object - should match what was specified in Reset()
+	 * @return A result structure that may or may not already have a result populated
+	 */
+	UNIVERSALOBJECTLOCATOR_API FResolveResult AsyncUnload(const UObject* Context = nullptr) const;
+
+	/**
+	 * Attempt to find the object this locator points to.
+	 * Shorthand for Resolve(FResolveParams::SyncFind(Context)).AsObject().SyncGet().Object.
+	 *
+	 * @param Context          (Optional) An optional context to use for finding the object - should match what was specified in Reset() or in construction
+	 * @return The located object, or nullptr on failure
+	 */
+	UNIVERSALOBJECTLOCATOR_API UObject* SyncFind(const UObject* Context = nullptr) const;
+
+	/**
+	 * Attempt to find the object this locator points to, loading it if necessary (and possible), and blocking until it is loaded.
+	 * Shorthand for Resolve(FResolveParams::SyncLoad(Context)).AsObject().SyncGet().Object.
+	 *
+	 * @param Context          (Optional) An optional context to use for finding/loading the object - should match what was specified in Reset()
+	 * @return The located object, or nullptr on failure
+	 */
+	UNIVERSALOBJECTLOCATOR_API UObject* SyncLoad(const UObject* Context = nullptr) const;
+
+	/**
+	 * Attempt to unload the object this locator points to if possible.
+	 * Shorthand for Resolve(FResolveParams::SyncUnload(Context)).AsVoid().SyncGet().
+	 *
+	 * @param Context          (Optional) An optional context to use for finding/loading the object - should match what was specified in Reset()
+	 * @return The located object, or nullptr on failure
+	 */
+	UNIVERSALOBJECTLOCATOR_API void SyncUnload(const UObject* Context = nullptr) const;
+
+	/**
+	 * Retrieve the fragment type relating to the last locator in this address
+	 */
+	UNIVERSALOBJECTLOCATOR_API const UE::UniversalObjectLocator::FFragmentType* GetLastFragmentType() const;
+
+	/**
+	 * Retrieve the fragment type handle relating to the last locator in this address
+	 */
+	UNIVERSALOBJECTLOCATOR_API UE::UniversalObjectLocator::FFragmentTypeHandle GetLastFragmentTypeHandle() const;
+
+	/**
+	 * Convert this locator to its string representation
+	 *
+	 * @param OutString        String builder to populate
+	 */
+	UNIVERSALOBJECTLOCATOR_API void ToString(FStringBuilderBase& OutString) const;
+
+	/**
+	 * Attempt to deserialize this locator from a string
+	 *
+	 * @param InString         The string to parse
+	 * @param InParams         Additional string parameters
+	 * @return Parse result, specifying success or failure, and number of characters that were parsed
+	 */
+	UNIVERSALOBJECTLOCATOR_API FParseStringResult TryParseString(FStringView InString, const FParseStringParams& InParams);
+
+	/**
+	 * Attempt to deserialize a new locator from a string.
+	 * Shorthand for FUniversalObjectLocator L; L.TryParseString(InString, InParams);.
+	 *
+	 * @param InString         The string to parse
+	 * @param InParams         Additional string parameters
+	 * @return A (perhaps empty) object locator
+	 */
+	static UNIVERSALOBJECTLOCATOR_API FUniversalObjectLocator FromString(FStringView InString, const FParseStringParams& InParams);
+
+public:
+
+	/**
+	 * Reset this locator back to its default-constructed, empty state
+	 */
+	UNIVERSALOBJECTLOCATOR_API void Reset();
+
+	/**
+	 * Attempt to reset this locator to point at a different object.
+	 * May result in an empty locator if no suitable address could be created.
+	 *
+	 * @param Object           The object that this locator should represent
+	 * @param Context          (Optional) Constrain this universal reference based on the specified context. This context should be passed to Resolve otherwise the resolution may fail.
+	 * @param StopAtContext    (Optional) Stop constructing this universal reference when we reach the following context (can be used if that context is always passed to Resolve to keep this type smaller)
+	 */
+	UNIVERSALOBJECTLOCATOR_API void Reset(UObject* Object, const UObject* Context = nullptr, const UObject* StopAtContext = nullptr);
+
+	/**
+	 * Retrieve the last fragment in this address
+	 */
+	UNIVERSALOBJECTLOCATOR_API FUniversalObjectLocatorFragment* GetLastFragment();
+
+	/**
+	 * Retrieve the last fragment in this address
+	 */
+	UNIVERSALOBJECTLOCATOR_API const FUniversalObjectLocatorFragment* GetLastFragment() const;
+
+	/**
+	 * Equality comparison.
+	 * @note: This tests for exact equality of its piecewise fragments. Equivalent but not equal locators will always return false.
+	 */
+	UNIVERSALOBJECTLOCATOR_API friend bool operator==(const FUniversalObjectLocator& A, const FUniversalObjectLocator& B);
+
+	/**
+	 * Inequality comparison.
+	 * @note: This tests for exact inequality of its piecewise fragments. Equivalent but not equal locators will always return true.
+	 */
+	UNIVERSALOBJECTLOCATOR_API friend bool operator!=(const FUniversalObjectLocator& A, const FUniversalObjectLocator& B);
+
+	/**
+	 * Type hashable
+	 */
+	UNIVERSALOBJECTLOCATOR_API friend uint32 GetTypeHash(const FUniversalObjectLocator& Locator);
+
+	/*~ Begin TStructOpsTypeTraits implementation */
+	UNIVERSALOBJECTLOCATOR_API bool ExportTextItem(FString& ValueStr, const FUniversalObjectLocator& DefaultValue, UObject* Parent, int32 PortFlags, UObject* ExportRootScope) const;
+	UNIVERSALOBJECTLOCATOR_API bool ImportTextItem(const TCHAR*& Buffer, int32 PortFlags, UObject* Parent, FOutputDevice* ErrorText, FArchive* InSerializingArchive = nullptr);
+	UNIVERSALOBJECTLOCATOR_API bool SerializeFromMismatchedTag(const FPropertyTag& Tag, FStructuredArchive::FSlot Slot);
+	/*~ End TStructOpsTypeTraits implementation */
+
+private:
+
+	/**  */
+	bool AddFragment(const UObject* Object, const UObject* Context, const UObject* StopAtContext);
+
+	FResolveResult ResolveSyncImpl(const FResolveParams& Params) const;
+	FResolveResult ResolveAsyncImpl(const FResolveParams& Params) const;
+
+private:
+
+	/** Array of relative locators ordered sequentially from outer to inner. The first locator is probably 'absolute' and is resolved with no context, although that is not a hard restriction */
+	UPROPERTY()
+	TArray<FUniversalObjectLocatorFragment> Fragments;
+};
+
+template<>
+struct TStructOpsTypeTraits<FUniversalObjectLocator> : public TStructOpsTypeTraitsBase2<FUniversalObjectLocator>
+{
+	enum
+	{
+		WithIdenticalViaEquality = true,
+		WithExportTextItem = true,
+		WithImportTextItem = true,
+		WithStructuredSerializeFromMismatchedTag = true,
+	};
+};
+
diff --git a/Engine/Source/Runtime/UniversalObjectLocator/Public/UniversalObjectLocatorFragment.h b/Engine/Source/Runtime/UniversalObjectLocator/Public/UniversalObjectLocatorFragment.h
new file mode 100644
--- /dev/null
+++ b/Engine/Source/Runtime/UniversalObjectLocator/Public/UniversalObjectLocatorFragment.h
@@ -0,0 +1,428 @@
+// Copyright Epic Games, Inc. All Rights Reserved.
+
+#pragma once
+
+#include "CoreTypes.h"
+#include "UObject/Class.h"
+#include "UniversalObjectLocatorFwd.h"
+#include "UniversalObjectLocatorFragmentType.h"
+#include "UniversalObjectLocatorFragmentTypeHandle.h"
+#include "Templates/Function.h"
+#include "UniversalObjectLocatorFragment.generated.h"
+
+/**
+ * Universal Object Locator (UOL) Fragments provide an extensible mechanism for referencing permanent, transient
+ *   or dynamically created objects relative to an external context. UOLs comprise zero or more nested fragments.
+ * 
+ * Creation and resolution of a fragment requires a context to be provided;
+ *   normally this will be the object on which the UOL exists as a property.
+ *
+ * The way in which the object is referenced is defined by globally registered 'FragmentTypes' 
+ *   (See IUniversalObjectLocatorModule::RegisterFragmentType). Each FragmentType can be thought of as somewhat
+ *   equivalent to a www URI fragment type, though the 'path' is not necessarily just a string, but includes
+ *   support for the full set of Engine Property types.
+ * 
+ * The type is implemented as a type-erased payload block, a fragment type handle and some internal flags.
+ * Payloads will be allocated using the inline memory if alignment and size constraints allow, but
+ *   will fall back to a heap allocation if necessary. Allocation should be avoided by keeping payload
+ *   types small.
+ *
+ * Aligned to 8 bytes, 32 (runtime) or 64 (editor) bytes big.
+ */
+USTRUCT(BlueprintType)
+struct alignas(8) FUniversalObjectLocatorFragment
+{
+	GENERATED_BODY()
+
+	using FParseStringResult = UE::UniversalObjectLocator::FParseStringResult;
+	using FParseStringParams = UE::UniversalObjectLocator::FParseStringParams;
+
+	static constexpr FAsciiSet ValidFragmentTypeCharacters = "abcdefghijklmnopqrstuvwxyzABCDEFGHIJKLMNOPQRSTUVWXYZ0123456789_-";
+	// Valid fragment delimiters are those that are allowable by RFC3986 including '/' and excluding '&' which we use to separate fragments within a path
+	static constexpr FAsciiSet ValidFragmentDelimiters = "!&'()*+,;=";
+	static constexpr FAsciiSet ValidFragmentPayloadCharacters = ValidFragmentTypeCharacters | ValidFragmentDelimiters;
+
+	/** Make our inline data buffer larger in-editor to support editor-only data without allocation */
+#if WITH_EDITORONLY_DATA
+	static constexpr SIZE_T Size = 64;
+#else
+	static constexpr SIZE_T Size = 32;
+#endif
+
+
+	/**
+	 * Construct a new fragment with a specific fragment type and data.
+	 * Used when a specific type of relative fragment is required.
+	 *
+	 * @param InHandle    A typed handle to the fragment type to use for construction. Retrieved from IUniversalObjectLocatorModule::RegisterFragmentType.
+	 * @param InArgs      (Optional) Payload construction arguments to be passed to T on construction. Omission implies default construction.
+	 */
+	template<typename T, typename ...ArgTypes>
+	FUniversalObjectLocatorFragment(UE::UniversalObjectLocator::TFragmentTypeHandle<T> InHandle, ArgTypes&& ...InArgs);
+
+
+	/**
+	 * Construct a new fragment with a specific fragment type and default-constructed payload.
+	 * Used when a specific type of relative fragment is required.
+	 *
+	 * @param InFragmentType  The fragment type to use with this fragment
+	 */
+	UNIVERSALOBJECTLOCATOR_API FUniversalObjectLocatorFragment(const UE::UniversalObjectLocator::FFragmentType& InFragmentType);
+
+	/**
+	 * Construct this fragment by binding it to an object within a given context.
+	 * @note: This constructor can 'fail' and result in an Empty fragment if no suitable fragment type could be found for the object
+	 */
+	UNIVERSALOBJECTLOCATOR_API FUniversalObjectLocatorFragment(const UObject* InObject, const UObject* Context);
+
+	/** Default constructor: initializes to an empty fragment with no fragment type */
+	UNIVERSALOBJECTLOCATOR_API FUniversalObjectLocatorFragment();
+
+	/** Destructor - destructs the payload, and frees any heap allocation as necessary */
+	UNIVERSALOBJECTLOCATOR_API ~FUniversalObjectLocatorFragment();
+
+	/** Copy construction/assignment */
+	UNIVERSALOBJECTLOCATOR_API FUniversalObjectLocatorFragment(const FUniversalObjectLocatorFragment& RHS);
+	UNIVERSALOBJECTLOCATOR_API FUniversalObjectLocatorFragment& operator=(const FUniversalObjectLocatorFragment& RHS);
+
+	/** Move construction/assignment */
+	UNIVERSALOBJECTLOCATOR_API FUniversalObjectLocatorFragment(FUniversalObjectLocatorFragment&& RHS);
+	UNIVERSALOBJECTLOCATOR_API FUniversalObjectLocatorFragment& operator=(FUniversalObjectLocatorFragment&& RHS);
+
+	/** Equality comparison. Compares the fragment type and payload data. */
+	UNIVERSALOBJECTLOCATOR_API friend bool operator==(const FUniversalObjectLocatorFragment& A, const FUniversalObjectLocatorFragment& B);
+	/** Inequality comparison. Compares the fragment type and payload data. */
+	UNIVERSALOBJECTLOCATOR_API friend bool operator!=(const FUniversalObjectLocatorFragment& A, const FUniversalObjectLocatorFragment& B);
+	/** Type hashable */
+	UNIVERSALOBJECTLOCATOR_API friend uint32 GetTypeHash(const FUniversalObjectLocatorFragment& Fragment);
+public:
+
+	/**
+	 * Attempt to resolve this fragment by invoking the payload's 'Resolve' function
+	 *
+	 * @param Params           Resolution parameters, defining the context to resolve within, and the type of resolution to perform
+	 * @return A result structure defining the resolved object pointer, and associated flags
+	 */
+	UNIVERSALOBJECTLOCATOR_API UE::UniversalObjectLocator::FResolveResult Resolve(const UE::UniversalObjectLocator::FResolveParams& Params) const;
+
+	/**
+	 * Check whether this reference is empty.
+	 * @note: An empty fragment can never resolve to an object, but is distinct from, and not equal to, a populated fragment that does points to a non-existent or irretrievable object.
+	 */
+	bool IsEmpty() const
+	{
+		return !FragmentType.IsValid();
+	}
+
+public:
+
+	/**
+	 * Reset this fragment back to its default-constructed, empty state
+	 */
+	UNIVERSALOBJECTLOCATOR_API void Reset();
+
+	/**
+	 * Reset this fragment to point to a new object from the specified context
+	 * @note: If no suitable fragment type could be found for the object and context, results in an Empty fragment
+	 */
+	UNIVERSALOBJECTLOCATOR_API void Reset(const UObject* InObject, const UObject* Context);
+
+	/**
+	 * Reset this fragment to point to a new object from the specified context using a filtered set of fragment types
+	 * @note: If no suitable fragment type could be found for the object and context, results in an Empty fragment
+	 */
+	UNIVERSALOBJECTLOCATOR_API void Reset(const UObject* InObject, const UObject* Context, TFunctionRef<bool(UE::UniversalObjectLocator::FFragmentTypeHandle)> CanUseFragmentType);
+
+public:
+
+	/**
+	 * Convert this fragment to a string of the form fragment-id[=fragment-payload]
+	 *
+	 * @param OutString        String builder to populate
+	 */
+	UNIVERSALOBJECTLOCATOR_API void ToString(FStringBuilderBase& OutString) const;
+
+	/**
+	 * Attempt to initialize this fragment from a string of the form fragment-type-id[=payload]
+	 * The state of this instance will not be changed if this function returns false.
+	 *
+	 * @param InString         The string to parse.
+	 * @param InParams         Additional string parameters
+	 * @return Parse result, specifying success or failure, and number of characters that were parsed
+	 */
+	UNIVERSALOBJECTLOCATOR_API FParseStringResult TryParseString(FStringView InString, const FParseStringParams& InParams);
+
+	/**
+	 * Attempt to default initialize this fragment using a string that defines the type
+	 * The state of this instance will not be changed if this function returns false.
+	 *
+	 * @param InString         The string to parse.
+	 * @param InParams         Additional string parameters
+	 * @return Parse result, specifying success or failure, and number of characters that were parsed
+	 */
+	UNIVERSALOBJECTLOCATOR_API FParseStringResult TryParseFragmentType(FStringView InString, const FParseStringParams& InParams);
+
+	/**
+	 * Attempt to deserialize this fragment's payload from a string, based on its currently assigned type
+	 * The state of this instance will not be changed if this function returns false.
+	 *
+	 * @param InString         The string to parse.
+	 * @param InParams         Additional string parameters
+	 * @return Parse result, specifying success or failure, and number of characters that were parsed
+	 */
+	UNIVERSALOBJECTLOCATOR_API FParseStringResult TryParseFragmentPayload(FStringView InString, const FParseStringParams& InParams);
+
+public:
+
+	/**
+	 * Try and retrieve this fragment's payload as a specific type using its fragment type handle
+	 *
+	 * @param InType           The handle of the fragment type to this payload to. This is returned from IUniversalObjectLocatorModule::RegisterFragmentType on registration.
+	 * @return A mutable pointer to the payload, or nullptr if this fragment is empty or of a different type.
+	 */
+	template<typename T>
+	T* GetPayloadAs(UE::UniversalObjectLocator::TFragmentTypeHandle<T> InType);
+
+	/**
+	 * Try and retrieve this fragment's payload as a specific type using its fragment type handle
+	 *
+	 * @param InType           The handle of the fragment type to this payload to. This is returned from IUniversalObjectLocatorModule::RegisterFragmentType on registration.
+	 * @return A mutable pointer to the payload, or nullptr if this fragment is empty or of a different type.
+	 */
+	template<typename T>
+	const T* GetPayloadAs(UE::UniversalObjectLocator::TFragmentTypeHandle<T> InType) const;
+
+	/**
+	 * Try and retrieve this fragment's payload as a specific type using its fragment type handle
+	 *
+	 * @param InType           The handle of the fragment type to this payload to. This is returned from IUniversalObjectLocatorModule::RegisterFragmentType on registration.
+	 * @param OutData          Pointer to retrieve the resulting payload ptr
+	 * @return true on success, false if this fragment is empty or of a different type.
+	 */
+	template<typename T>
+	bool TryGetPayloadAs(UE::UniversalObjectLocator::TFragmentTypeHandle<T> InType, T*& OutData);
+
+	/**
+	 * Try and retrieve this fragment's payload as a specific type using its fragment type handle
+	 *
+	 * @param InType           The handle of the fragment type to this payload to. This is returned from IUniversalObjectLocatorModule::RegisterFragmentType on registration.
+	 * @param OutData          Pointer to retrieve the resulting payload ptr
+	 * @return true on success, false if this fragment is empty or of a different type.
+	 */
+	template<typename T>
+	bool TryGetPayloadAs(UE::UniversalObjectLocator::TFragmentTypeHandle<T> InType, const T*& OutData) const;
+
+	/**
+	 * Retrieve this fragment's payload data
+	 *
+	 * @return The payload data or nullptr if this fragment is empty
+	 */
+	UNIVERSALOBJECTLOCATOR_API void* GetPayload();
+
+	/**
+	 * Retrieve this fragment's payload data
+	 *
+	 * @return The payload data or nullptr if this fragment is empty
+	 */
+	UNIVERSALOBJECTLOCATOR_API const void* GetPayload() const;
+
+	/**
+	 * Retrieve this fragment's fragment type
+	 *
+	 * @return The fragment type or nullptr if this fragment is empty
+	 */
+	UNIVERSALOBJECTLOCATOR_API const UE::UniversalObjectLocator::FFragmentType* GetFragmentType() const;
+
+	/**
+	 * Retrieve this fragment's fragment struct type
+	 *
+	 * @return The fragment struct type or nullptr if this fragment is empty
+	 */
+	UNIVERSALOBJECTLOCATOR_API UScriptStruct* GetFragmentStruct() const;
+
+	/**
+	 * Retrieve this fragment's fragment type handle
+	 *
+	 * @return The fragment type handle (possibly invalid if this fragment is empty)
+	 */
+	UNIVERSALOBJECTLOCATOR_API UE::UniversalObjectLocator::FFragmentTypeHandle GetFragmentTypeHandle() const;
+
+public:
+
+	/*~ Begin TStructOpsTypeTraits implementation */
+	UNIVERSALOBJECTLOCATOR_API bool Serialize(FArchive& Ar);
+	UNIVERSALOBJECTLOCATOR_API void AddStructReferencedObjects(FReferenceCollector& Collector);
+	UNIVERSALOBJECTLOCATOR_API bool ExportTextItem(FString& ValueStr, const FUniversalObjectLocatorFragment& DefaultValue, UObject* Parent, int32 PortFlags, UObject* ExportRootScope) const;
+	UNIVERSALOBJECTLOCATOR_API bool ImportTextItem(const TCHAR*& Buffer, int32 PortFlags, UObject* Parent, FOutputDevice* ErrorText, FArchive* InSerializingArchive = nullptr);
+	UNIVERSALOBJECTLOCATOR_API bool SerializeFromMismatchedTag(const FPropertyTag& Tag, FStructuredArchive::FSlot Slot);
+	UNIVERSALOBJECTLOCATOR_API void GetPreloadDependencies(TArray<UObject*>& OutDeps);
+	/*~ End TStructOpsTypeTraits implementation */
+
+protected:
+
+
+	/** Runtime type-checking against specific struct types. Compiled out if checks are not enabled. */
+#if DO_CHECK
+	/** Runtime type-checking against a specific struct type. Compiled out if checks are not enabled. */
+	UNIVERSALOBJECTLOCATOR_API void CheckPayloadType(UScriptStruct* TypeToCompare) const;
+#else
+	FORCEINLINE static constexpr void CheckPayloadType(void* TypeToCompare)
+	{
+	}
+#endif
+
+protected:
+
+	/**
+	 * Default-initialize the fragment payload using the specified type
+	 */
+	void InitializePayload(const UScriptStruct* StructType);
+
+	/**
+	 * Destroy the payload (if valid) by calling its destructor and freeing the memory (if necessary)
+	 */
+	void DestroyPayload();
+
+private:
+
+	/*~ Utility symbol name to guarantee that FFragmentType can be resolved within the context of a FUniversalObjectLocatorFragment within natvis expressions */
+	struct FDebuggableFragmentType : UE::UniversalObjectLocator::FFragmentType
+	{
+	};
+
+	/*
+	 * Payload data - implicitly aligned to a 8 byte boundary since it's the first member.
+	 * Given payload type T, this is either a type-erased T() value (where bIsInline==1),
+	 *    or a T* to a heap allocated T (where bIsInline==0)
+	 * Size is specifically defined by the desired overall size of FUniversalObjectLocatorFragment::Size, minus space for other members
+	 */
+	uint8 Data[Size-2];
+
+	/** 1 Byte - the fragment type portion of the universal fragment */
+	UE::UniversalObjectLocator::FFragmentTypeHandle FragmentType;
+
+	/*~ 1 Byte of flags */
+	/** True when FragmentType has been assigned to a valid handle, and Data has been initialized with FragmentType::PayloadType */
+	uint8 bIsInitialized : 1;
+	/** True if Data is an inline allocation of FragmentType::PayloadType, false means Data is a (void*) to the heap allocated data. */
+	uint8 bIsInline : 1;
+};
+
+template<>
+struct TStructOpsTypeTraits<FUniversalObjectLocatorFragment> : public TStructOpsTypeTraitsBase2<FUniversalObjectLocatorFragment>
+{
+	enum
+	{
+		WithSerializer = true,
+		WithCopy = true,
+		WithIdenticalViaEquality = true,
+		WithExportTextItem = true,
+		WithImportTextItem = true,
+		WithAddStructReferencedObjects = true,
+	};
+};
+
+
+template<typename T, typename ...ArgTypes>
+FUniversalObjectLocatorFragment::FUniversalObjectLocatorFragment(UE::UniversalObjectLocator::TFragmentTypeHandle<T> InHandle, ArgTypes&& ...InArgs)
+	: FragmentType(InHandle)
+	, bIsInitialized(1)
+{
+	checkf(InHandle, TEXT("Attempting to construct a new fragment from an invalid fragment type handle - was it registered?"));
+
+	bIsInline = sizeof(T) <= sizeof(Data) && alignof(T) <= alignof(FUniversalObjectLocatorFragment);
+	if (!bIsInline)
+	{
+		// We have to allocate this struct on the heap
+		void* HeapAllocation = FMemory::Malloc(sizeof(T), alignof(T));
+		*reinterpret_cast<void**>(Data) = HeapAllocation;
+	}
+
+	// Placement new the payload
+	new (GetPayload()) T(Forward<ArgTypes>(InArgs)...);
+}
+
+template<typename T>
+T* FUniversalObjectLocatorFragment::GetPayloadAs(UE::UniversalObjectLocator::TFragmentTypeHandle<T> InType)
+{
+	if (FragmentType.IsValid() && ensureMsgf(FragmentType == InType, TEXT("Type mismatch when accessing payload data!")))
+	{
+		return static_cast<T*>(GetPayload());
+	}
+	return nullptr;
+}
+
+template<typename T>
+const T* FUniversalObjectLocatorFragment::GetPayloadAs(UE::UniversalObjectLocator::TFragmentTypeHandle<T> InType) const
+{
+	return const_cast<FUniversalObjectLocatorFragment*>(this)->GetPayloadAs(InType);
+}
+
+template<typename T>
+bool FUniversalObjectLocatorFragment::TryGetPayloadAs(UE::UniversalObjectLocator::TFragmentTypeHandle<T> InType, T*& OutData)
+{
+	if (FragmentType.IsValid() && FragmentType == InType)
+	{
+		OutData = static_cast<T*>(GetPayload());
+		return true;
+	}
+	return false;
+}
+
+template<typename T>
+bool FUniversalObjectLocatorFragment::TryGetPayloadAs(UE::UniversalObjectLocator::TFragmentTypeHandle<T> InType, const T*& OutData) const
+{
+	if (FragmentType.IsValid() && FragmentType == InType)
+	{
+		OutData = static_cast<const T*>(GetPayload());
+		return true;
+	}
+	return false;
+}
+
+
+/** Empty struct type used for deserializing unknown fragment type payloads */
+USTRUCT()
+struct FUniversalObjectLocatorEmptyPayload
+{
+	GENERATED_BODY()
+};
+
+
+template<typename PayloadType>
+struct TUniversalObjectLocatorFragment : FUniversalObjectLocatorFragment
+{
+	TUniversalObjectLocatorFragment()
+		: FUniversalObjectLocatorFragment(PayloadType::FragmentType)
+	{
+	}
+
+	TUniversalObjectLocatorFragment(UE::UniversalObjectLocator::TFragmentTypeHandle<PayloadType> InHandle)
+		: FUniversalObjectLocatorFragment(InHandle)
+	{
+	}
+
+	template<typename ...ArgTypes>
+	TUniversalObjectLocatorFragment(ArgTypes&& ...InArgs)
+		: FUniversalObjectLocatorFragment(PayloadType::FragmentType, Forward<ArgTypes>(InArgs)...)
+	{
+	}
+
+	template<typename ...ArgTypes>
+	TUniversalObjectLocatorFragment(UE::UniversalObjectLocator::TFragmentTypeHandle<PayloadType> InHandle, ArgTypes&& ...InArgs)
+		: FUniversalObjectLocatorFragment(InHandle, Forward<ArgTypes>(InArgs)...)
+	{
+	}
+
+	PayloadType* GetPayload()
+	{
+		CheckPayloadType(PayloadType::StaticStruct());
+		return static_cast<PayloadType*>(FUniversalObjectLocatorFragment::GetPayload());
+	}
+
+	const PayloadType* GetPayload() const
+	{
+		CheckPayloadType(PayloadType::StaticStruct());
+		return static_cast<const PayloadType*>(FUniversalObjectLocatorFragment::GetPayload());
+	}
+};
diff --git a/Engine/Source/Runtime/UniversalObjectLocator/Public/UniversalObjectLocatorFragmentType.h b/Engine/Source/Runtime/UniversalObjectLocator/Public/UniversalObjectLocatorFragmentType.h
new file mode 100644
--- /dev/null
+++ b/Engine/Source/Runtime/UniversalObjectLocator/Public/UniversalObjectLocatorFragmentType.h
@@ -0,0 +1,152 @@
+// Copyright Epic Games, Inc. All Rights Reserved.
+
+#pragma once
+
+#include "CoreTypes.h"
+#include "Misc/AssertionMacros.h"
+#include "Misc/AsciiSet.h"
+#include "Misc/StringBuilder.h"
+#include "Containers/PagedArray.h"
+#include "Containers/StringView.h"
+
+#include "UObject/ObjectPtr.h"
+#include "UObject/Class.h"
+
+#include "UniversalObjectLocatorFwd.h"
+#include "UniversalObjectLocatorResolveParams.h"
+
+namespace UE::UniversalObjectLocator
+{
+
+struct IPayloadDebugging;
+struct IFragmentTypeDebuggingAssistant;
+
+/**
+ * A Universal Object Locator fragment type defines a specific mechanism for resolving an object that is referenced
+ *   inside a Universal Object Locator (which can be considered a chain of fragments). Each fragment
+ *   contains the attributes defining its payload, and how to interact that payload.
+ *
+ * In order to define a new fragment type, consumers must provide a UStruct 'payload' that defines the data required
+ *   for the reference. The struct also defines how to interact with the payload via the following functions:
+ *
+ *   // Initialize the fragment to point at an object within a context
+ *   FInitializeResult Initialize(const FInitializeParams& InParams) const;
+ *   // Resolve the fragment to the resulting object, potentially loading or unloading it
+ *   FResolveResult Resolve(const FResolveParams& Params) const;
+ *
+ *   // Convert this fragment to its string representation
+ *   void ToString(FStringBuilderBase& OutStringBuilder) const;
+ *   // Try and parse this fragment from its string representation
+ *   FParseStringResult TryParseString(FStringView InString, const FParseStringParams& Params) const;
+ *
+ *   // Compute a priority for this fragment based on an object and context in order to decide which fragment type should be used if multiple can address the same object
+ *   static uint32 ComputePriority(const UObject* Object, const UObject* Context) const;
+ * 
+ * Types must also be UStructs that are equality comparable, and GetTypeHashable.
+ * These functions are bound as function pointers on construction, allowing them potentially to be inlined into the static dispatchers.
+ */
+struct FFragmentType
+{
+	FResolveResult ResolvePayload(const void* Payload, const FResolveParams& Params) const;
+	FInitializeResult InitializePayload(void* Payload, const FInitializeParams& InParams) const;
+	void ToString(const void* Payload, FStringBuilderBase& OutStringBuilder) const;
+	FParseStringResult TryParseString(void* Payload, FStringView InString, const FParseStringParams& Params) const;
+
+	uint32 ComputePriority(const UObject* Object, const UObject* Context) const;
+
+	using FResolveCallback    = FResolveResult (*)(const void* /* Payload */, const FResolveParams& /* Params */);
+	using FInitializeCallback = FInitializeResult (*)(void* /* Payload */, const FInitializeParams&);
+	using FToString           = void (*)(const void* /* Payload */, FStringBuilderBase& /* OutStringBuilder */);
+	using FTryParseString     = FParseStringResult (*)(void* /* Payload */, FStringView /* InString */, const FParseStringParams& /* Params */);
+
+	using FPriorityCallback   = uint32 (*)(const UObject* /* Object */, const UObject* /* Context */);
+
+	/** Bindings structure for function pointers that directly invoke behavior on fragment type payloads */
+	struct
+	{
+		/** Function binding for initializing a payload with this fragment type */
+		FInitializeCallback Initialize;
+		/** Function binding for resolving a payload for this fragment type */
+		FResolveCallback    Resolve;
+		/** Function binding for serializing this payload to a string */
+		FToString           ToString;
+		/** Function binding for parsing this payload from a string */
+		FTryParseString     TryParseString;
+	} InstanceBindings;
+
+	/** Bindings structure for function pointers that directly invoke behavior on fragment type payloads */
+	struct
+	{
+		/** static function binding for computing the priority of this fragment type */
+		FPriorityCallback   Priority;
+	} StaticBindings;
+
+	/** Struct pointer to the payload type */
+	TObjectPtr<const UScriptStruct> PayloadType;
+
+	/** Debugging assistant structure used for natvis visualizations */
+	TSharedPtr<IFragmentTypeDebuggingAssistant> DebuggingAssistant;
+
+	/** Descriptive display text for this type of fragment type */
+	FText DisplayText;
+
+	/** Unique identifier for this fragment type. This is serialized persistently to identify the fragment type, and exists as part of the fragment type's string representation */
+	FName FragmentTypeID;
+
+	/** Name of the primary editor type for this fragment type, defining how it appears on UI */
+	FName PrimaryEditorType;
+
+	UScriptStruct* GetStruct() const
+	{
+		return const_cast<UScriptStruct*>(PayloadType.Get());
+	}
+};
+
+template<typename T>
+struct TFragmentType : FFragmentType
+{
+};
+
+struct IPayloadDebugging
+{
+	virtual ~IPayloadDebugging()
+	{
+	}
+};
+
+template<typename T>
+struct TLocatorFragment : IPayloadDebugging
+{
+	explicit TLocatorFragment(const T* InPayload)
+		: Payload(InPayload)
+	{
+	}
+
+	const T* Payload;
+};
+
+struct IFragmentTypeDebuggingAssistant : TSharedFromThis<IFragmentTypeDebuggingAssistant>
+{
+	virtual ~IFragmentTypeDebuggingAssistant()
+	{
+	}
+	virtual IPayloadDebugging* MakePayloadVisualizer(const void*) = 0;
+	virtual void Purge() = 0;
+};
+
+template<typename T>
+struct TFragmentTypeDebuggingAssistant : IFragmentTypeDebuggingAssistant
+{
+	IPayloadDebugging* MakePayloadVisualizer(const void* InPayload) override
+	{
+		return &Visualizers.Emplace_GetRef(static_cast<const T*>(InPayload));
+	}
+	void Purge() override
+	{
+		Visualizers.Empty();
+	}
+
+	TPagedArray<TLocatorFragment<T>> Visualizers;
+};
+
+} // namespace UE::UniversalObjectLocator
\ No newline at end of file
diff --git a/Engine/Source/Runtime/UniversalObjectLocator/Public/UniversalObjectLocatorFragmentTypeHandle.h b/Engine/Source/Runtime/UniversalObjectLocator/Public/UniversalObjectLocatorFragmentTypeHandle.h
new file mode 100644
--- /dev/null
+++ b/Engine/Source/Runtime/UniversalObjectLocator/Public/UniversalObjectLocatorFragmentTypeHandle.h
@@ -0,0 +1,105 @@
+// Copyright Epic Games, Inc. All Rights Reserved.
+
+#pragma once
+
+#include "CoreTypes.h"
+#include "Templates/TypeHash.h"
+#include "UniversalObjectLocatorFwd.h"
+
+namespace UE::UniversalObjectLocator
+{
+
+/**
+ * A handled to a globally registered Universal Object Locator fragment type.
+ *
+ * Registered through IUniversalObjectLocatorModule::RegisterFragmentType, fragment types define all the 
+ *   attributes for a specific type of object resolution protocol that can be represented in a UOL.
+ */
+struct FFragmentTypeHandle
+{
+	FFragmentTypeHandle()
+		: Handle(0xff)
+	{
+	}
+
+	explicit FFragmentTypeHandle(uint8 InHandleOffset)
+		: Handle(InHandleOffset)
+	{
+		check(InHandleOffset != 0xff);
+	}
+
+	explicit operator bool() const
+	{
+		return IsValid();
+	}
+
+	bool IsValid() const
+	{
+		return Handle != 0xff;
+	}
+
+	uint8 GetIndex() const
+	{
+		return Handle;
+	}
+
+	friend uint32 GetTypeHash(FFragmentTypeHandle In)
+	{
+		return GetTypeHash(In.Handle);
+	}
+
+	friend bool operator==(FFragmentTypeHandle A, FFragmentTypeHandle B)
+	{
+		return A.Handle == B.Handle;
+	}
+
+	friend bool operator!=(FFragmentTypeHandle A, FFragmentTypeHandle B)
+	{
+		return A.Handle != B.Handle;
+	}
+
+	friend bool operator<(FFragmentTypeHandle A, FFragmentTypeHandle B)
+	{
+		return A.Handle < B.Handle;
+	}
+
+	friend bool operator>(FFragmentTypeHandle A, FFragmentTypeHandle B)
+	{
+		return A.Handle > B.Handle;
+	}
+
+	UNIVERSALOBJECTLOCATOR_API FFragmentType* Resolve() const;
+
+private:
+	/** Global index into UE::UniversalObjectLocator::GFragmentTypes */
+	uint8 Handle;
+};
+
+
+/**
+ * Typed version of FFragmentTypeHandle that provides typed access to the underlying fragment type
+ */
+template<typename T>
+struct TFragmentTypeHandle : FFragmentTypeHandle
+{
+	TFragmentTypeHandle()
+	{
+	}
+
+	TFragmentType<T>* Resolve() const
+	{
+		return static_cast<TFragmentType<T>*>(FFragmentTypeHandle::Resolve());
+	}
+
+private:
+
+	friend class IUniversalObjectLocatorModule;
+	friend class FUniversalObjectLocatorModule;
+
+	TFragmentTypeHandle(FFragmentTypeHandle InHandle)
+		: FFragmentTypeHandle(InHandle)
+	{
+	}
+};
+
+} // namespace UE::UniversalObjectLocator
\ No newline at end of file
diff --git a/Engine/Source/Runtime/UniversalObjectLocator/Public/UniversalObjectLocatorFwd.h b/Engine/Source/Runtime/UniversalObjectLocator/Public/UniversalObjectLocatorFwd.h
new file mode 100644
--- /dev/null
+++ b/Engine/Source/Runtime/UniversalObjectLocator/Public/UniversalObjectLocatorFwd.h
@@ -0,0 +1,32 @@
+// Copyright Epic Games, Inc. All Rights Reserved.
+
+#pragma once
+
+#include "HAL/Platform.h"
+
+struct FUniversalObjectLocator;
+struct FUniversalObjectLocatorFragment;
+
+template<typename> struct TUniversalObjectLocatorFragment;
+
+namespace UE::UniversalObjectLocator
+{
+
+enum class EResolveFlags : uint8;
+enum class ELocatorType : uint8;
+
+struct FFragmentType;
+template<typename> struct TFragmentType;
+
+struct FFragmentTypeHandle;
+template<typename> struct TFragmentTypeHandle;
+
+struct FInitializeParams;
+struct FInitializeResult;
+struct FResolveParams;
+struct FResolveResultFlags;
+struct FResolveResult;
+struct FParseStringParams;
+struct FParseStringResult;
+
+} // namespace UE::UniversalObjectLocator
\ No newline at end of file
diff --git a/Engine/Source/Runtime/UniversalObjectLocator/Public/UniversalObjectLocatorInitializeParams.h b/Engine/Source/Runtime/UniversalObjectLocator/Public/UniversalObjectLocatorInitializeParams.h
new file mode 100644
--- /dev/null
+++ b/Engine/Source/Runtime/UniversalObjectLocator/Public/UniversalObjectLocatorInitializeParams.h
@@ -0,0 +1,38 @@
+// Copyright Epic Games, Inc. All Rights Reserved.
+
+#pragma once
+
+#include "CoreTypes.h"
+
+class UObject;
+
+namespace UE::UniversalObjectLocator
+{
+
+/**
+ * Fragment initialization parameter structure
+ */
+struct FInitializeParams
+{
+	/** The object that should be referenced by this locator (or, in the case of transient objects, an equivalent object). */
+	const UObject* Object = nullptr;
+
+	/** The context within which Object will be resolved within. Maybe null. */
+	const UObject* Context = nullptr;
+
+	/** Retrieve the context as a specific type. */
+	template<typename T>
+	T* GetObjectAs() const
+	{
+		return Cast<T>(Object);
+	}
+
+	/** Retrieve the context as a specific type. */
+	template<typename T>
+	const T* GetContextAs() const
+	{
+		return Cast<const T>(Context);
+	}
+};
+
+} // namespace UE::UniversalObjectLocator
\ No newline at end of file
diff --git a/Engine/Source/Runtime/UniversalObjectLocator/Public/UniversalObjectLocatorInitializeResult.h b/Engine/Source/Runtime/UniversalObjectLocator/Public/UniversalObjectLocatorInitializeResult.h
new file mode 100644
--- /dev/null
+++ b/Engine/Source/Runtime/UniversalObjectLocator/Public/UniversalObjectLocatorInitializeResult.h
@@ -0,0 +1,56 @@
+// Copyright Epic Games, Inc. All Rights Reserved.
+
+#pragma once
+
+#include "CoreTypes.h"
+#include "Misc/AssertionMacros.h"
+
+class UObject;
+
+namespace UE::UniversalObjectLocator
+{
+
+enum class ELocatorType : uint8
+{
+	Absolute,
+	Relative,
+	Undefined,
+};
+
+/**
+ * Fragment initialization result structure, specfying the type of locator that was created and an optional relative 'parent'
+ */
+struct FInitializeResult
+{
+	/** Optional object that this fragment should be resolved relative to */
+	const UObject* RelativeToContext = nullptr;
+
+	/** The type of the locator fragment */
+	ELocatorType Type = ELocatorType::Undefined;
+
+	static FInitializeResult Absolute()
+	{
+		FInitializeResult Result;
+		Result.Type = ELocatorType::Absolute;
+		return Result;
+	}
+
+	static FInitializeResult Relative(const UObject* InRelativeToContext)
+	{
+		check(InRelativeToContext);
+
+		FInitializeResult Result;
+		Result.RelativeToContext = InRelativeToContext;
+		Result.Type = ELocatorType::Absolute;
+		return Result;
+	}
+
+	static FInitializeResult Failure()
+	{
+		FInitializeResult Result;
+		Result.Type = ELocatorType::Undefined;
+		return Result;
+	}
+};
+
+} // namespace UE::UniversalObjectLocator
\ No newline at end of file
diff --git a/Engine/Source/Runtime/UniversalObjectLocator/Public/UniversalObjectLocatorResolveParams.h b/Engine/Source/Runtime/UniversalObjectLocator/Public/UniversalObjectLocatorResolveParams.h
new file mode 100644
--- /dev/null
+++ b/Engine/Source/Runtime/UniversalObjectLocator/Public/UniversalObjectLocatorResolveParams.h
@@ -0,0 +1,418 @@
+// Copyright Epic Games, Inc. All Rights Reserved.
+
+#pragma once
+
+#include "CoreTypes.h"
+#include "Async/Future.h"
+#include "Misc/TVariant.h"
+
+class UObject;
+
+namespace UE::UniversalObjectLocator
+{
+
+
+enum class EResolveFlags : uint8
+{
+	None,
+
+	/** Flag to indicate whether the object should be loaded if it is not currently findable */
+	Load = 1 << 0,
+
+	/** Flag to indicate whether the object should be unloaded or destroyed. Mutually exclusive with bLoad. */
+	Unload = 1 << 1,
+
+	/**
+	 * Indicates that the operation should be performed asynchronously if possible.
+	 *   When not combined with WillWait, the caller will never block waiting for the result.
+	 *   When combined with WillWait, the caller will block on this thread until the result is available,
+	 *    so care needs to be taken avoid a deadlock if there are additional threading constraints on the load. */
+	Async = 1 << 2,
+
+	/** Indicates the calling code is going to block waiting for the result. */
+	WillWait = 1 << 3,
+
+	/** Combination of Async and WillWait. */
+	AsyncWait = Async | WillWait,
+};
+ENUM_CLASS_FLAGS(EResolveFlags)
+
+
+/**
+ * Parameters required to resolve a universal object locator
+ */
+struct FResolveParams
+{
+	FResolveParams()
+		: Context(nullptr)
+		, Flags(EResolveFlags::None)
+	{
+	}
+
+	FResolveParams(const UObject* InContext)
+		: Context(InContext)
+		, Flags(EResolveFlags::None)
+	{
+	}
+
+	FResolveParams(const UObject* InContext, EResolveFlags InFlags)
+		: Context(InContext)
+		, Flags(InFlags)
+	{
+	}
+
+	/** 
+	 * Utility function to indicate that the object should be found, but not loaded or created if it is not currently
+	 */
+	static FResolveParams AsyncFind(const UObject* InContext = nullptr)
+	{
+		return FResolveParams(InContext, EResolveFlags::Async);
+	}
+
+	/** 
+	 * Utility function to indicate that the object should be loaded (or dynamically created) if it is not currently
+	 */
+	static FResolveParams AsyncLoad(const UObject* InContext = nullptr)
+	{
+		return FResolveParams(InContext, EResolveFlags::Async | EResolveFlags::Load);
+	}
+
+	/** 
+	 * Utility function to indicate that the object should be unloaded (or dynamically destroyed) if it is not currently
+	 */
+	static FResolveParams AsyncUnload(const UObject* InContext)
+	{
+		return FResolveParams(InContext, EResolveFlags::Async | EResolveFlags::Unload);
+	}
+
+	/** 
+	 * Utility function to indicate that the object should be found, but not loaded or created if it is not currently
+	 */
+	static FResolveParams SyncFind(const UObject* InContext = nullptr)
+	{
+		return FResolveParams(InContext);
+	}
+
+	/** 
+	 * Utility function to indicate that the object should be loaded (or dynamically created) if it is not currently
+	 */
+	static FResolveParams SyncLoad(const UObject* InContext = nullptr)
+	{
+		return FResolveParams(InContext, EResolveFlags::Load);
+	}
+
+	/** 
+	 * Utility function to indicate that the object should be unloaded (or dynamically destroyed) if it is not currently
+	 */
+	static FResolveParams SyncUnload(const UObject* InContext)
+	{
+		return FResolveParams(InContext, EResolveFlags::Unload);
+	}
+
+	/** (Optional) Object to use as a context for resolution. Normally this is the object that owns the reference being resolved */
+	const UObject* Context;
+
+	/** Flag structure */
+	EResolveFlags Flags;
+};
+
+
+/**
+ * Flag structure returned from an attempt to resolve a Universal Object Locator
+ */
+struct FResolveResultFlags
+{
+	FResolveResultFlags()
+		: bWasLoaded(0)
+		, bWasLoadedIndirectly(0)
+	{
+	}
+
+	/** Indicates the object was loaded as a result of this request */
+	uint8 bWasLoaded : 1;
+	/** Indicates that an outer object was loaded in order to fulfil this request, and the final object itself was loaded as a result */
+	uint8 bWasLoadedIndirectly : 1;
+};
+
+struct FResolveResultData
+{
+	FResolveResultData()
+		: Object(nullptr)
+	{
+	}
+
+	FResolveResultData(UObject* InObject, FResolveResultFlags InFlags = FResolveResultFlags())
+		: Object(InObject)
+		, Flags(InFlags)
+	{
+	}
+
+	/**
+	 * The resulting object or nullptr if it could not be found or loaded per the request params.
+	 * Should only be set if Params.Flags.bAsync is false.
+	 */
+	UObject* Object = nullptr;
+
+	/** Flags relating to the operation */
+	FResolveResultFlags Flags;
+};
+
+
+struct FResolveResult
+{
+	FResolveResult()
+		: Value(TInPlaceType<FResolveResultData>())
+	{
+	}
+
+	FResolveResult(FResolveResultData InValue)
+		: Value(TInPlaceType<FResolveResultData>(), MoveTemp(InValue))
+	{
+	}
+
+	FResolveResult(TFuture<FResolveResultData>&& InFuture)
+		: Value(TInPlaceType<TFuture<FResolveResultData>>(), MoveTemp(InFuture))
+	{
+	}
+
+	FResolveResult(const FResolveResult&) = delete;
+	void operator=(const FResolveResult&) = delete;
+
+	FResolveResult(FResolveResult&&) = default;
+	FResolveResult& operator=(FResolveResult&&) = default;
+
+public:
+
+	/**
+	 * Check whether this result is asynchronous
+	 */
+	bool IsAsync() const
+	{
+		return Value.IsType<TFuture<FResolveResultData>>();
+	}
+
+	/**
+	 * Release the future from this type if it is set.
+	 * WARNING: This can only be called once per instance. Invalidates this class.
+	 * @return This result's future, or an empty one if the result has already been applied.
+	 */
+	TFuture<FResolveResultData> ReleaseFuture()
+	{
+		if (TFuture<FResolveResultData>* Future = Value.TryGet<TFuture<FResolveResultData>>())
+		{
+			return MoveTemp(*Future);
+		}
+		return TFuture<FResolveResultData>();
+	}
+
+	/**
+	 * Check whether the result needs to be waited on.
+	 * @return true if the result hasn't been set yet: SyncGet will block, and SyncGetNoWait will return null. If false, SyncGet will always return immediately, and SyncGetNoWait will never return null.
+	 */
+	bool NeedsWait() const
+	{
+		if (const TFuture<FResolveResultData>* Future = Value.TryGet<TFuture<FResolveResultData>>())
+		{
+			return !Future->IsReady();
+		}
+		return false;
+	}
+
+	/**
+	 * Retrieve the result if it is currently available without waiting. Returns null if it is still being processed.
+	 */
+	const FResolveResultData* SyncGetNoWait()
+	{
+		const TFuture<FResolveResultData>* Future = Value.TryGet<TFuture<FResolveResultData>>();
+		if (Future)
+		{
+			if (!Future->IsReady())
+			{
+				return nullptr;
+			}
+
+			FResolveResultData Data = Future->Get();
+			Value.Emplace<FResolveResultData>(Data);
+		}
+
+		return &Value.Get<FResolveResultData>();
+	}
+
+	/**
+	 * Retrieve the result, blocking this thread if necessary.
+	 */
+	FResolveResultData SyncGet() const
+	{
+		const TFuture<FResolveResultData>* Future = Value.TryGet<TFuture<FResolveResultData>>();
+		if (Future)
+		{
+			return Future->Get();
+		}
+		else
+		{
+			return Value.Get<FResolveResultData>();
+		}
+	}
+
+	/**
+	 * Retrieve this result asynchronously.
+	 * WARNING: This can only be called once per instance.
+	 */
+	void AsyncGet(TUniqueFunction<void(const FResolveResultData&)>&& OnComplete)
+	{
+		TFuture<FResolveResultData>* Future = Value.TryGet<TFuture<FResolveResultData>>();
+		if (Future)
+		{
+			if (Future->IsReady())
+			{
+				OnComplete(Future->Get());
+			}
+			else
+			{
+				Future->Then(
+					[OnComplete = MoveTemp(OnComplete)](TFuture<FResolveResultData>&& InValue)
+					{
+						OnComplete(InValue.Get());
+					}
+				);
+			}
+		}
+		else
+		{
+			OnComplete(MoveTemp(Value.Get<FResolveResultData>()));
+		}
+	}
+
+private:
+
+	TVariant<
+		FResolveResultData,
+		TFuture<FResolveResultData>
+	> Value;
+};
+
+// struct FResolveObjectResult
+// {
+// 	FResolveObjectResult()
+// 		: Object(nullptr)
+// 	{
+// 	}
+
+// 	FResolveObjectResult(UObject* InObject, FResolveResultFlags InFlags = FResolveResultFlags())
+// 		: Object(InObject)
+// 		, Flags(InFlags)
+// 	{
+// 	}
+
+// 	/**
+// 	 * The resulting object or nullptr if it could not be found or loaded per the request params.
+// 	 * Should only be set if Params.Flags.bAsync is false.
+// 	 */
+// 	UObject* Object = nullptr;
+
+// 	/** Flags relating to the operation */
+// 	FResolveResultFlags Flags;
+// };
+
+
+// struct FResolveVoidResult
+// {
+// };
+
+// /**
+//  * Result structure that is returned from a request to resolve a Universal Object Locator
+//  */
+// struct FResolveResult
+// {
+// 	/**
+// 	 * Empty result
+// 	 */
+// 	FResolveResult()
+// 		: Value(TInPlaceType<TAsyncResolveResult<FResolveVoidResult>>())
+// 	{
+// 	}
+
+// 	/**
+// 	 * Constructor to use when an object needs to be returned from a Resolve function (ie, Load or Find)
+// 	 */
+// 	FResolveResult(FResolveObjectResult&& InResolveResult)
+// 		: Value(TInPlaceType<TAsyncResolveResult<FResolveObjectResult>>(), MoveTemp(InResolveResult))
+// 	{
+// 	}
+
+// 	/**
+// 	 * Constructor to use when no information needs to be returned from a Resolve function (ie, Unload)
+// 	 */
+// 	FResolveResult(FResolveVoidResult&& InVoidResult)
+// 		: Value(TInPlaceType<TAsyncResolveResult<FResolveVoidResult>>(), MoveTemp(InVoidResult))
+// 	{
+// 	}
+
+// 	/**
+// 	 * Constructor to use when an object needs to be returned from a Resolve function (ie, Load or Find)
+// 	 */
+// 	FResolveResult(TAsyncResolveResult<FResolveObjectResult>&& InResolveResult)
+// 		: Value(TInPlaceType<TAsyncResolveResult<FResolveObjectResult>>(), MoveTemp(InResolveResult))
+// 	{
+// 	}
+
+// 	/**
+// 	 * Constructor to use when no information needs to be returned from a Resolve function (ie, Unload)
+// 	 */
+// 	FResolveResult(TAsyncResolveResult<FResolveVoidResult>&& InVoidResult)
+// 		: Value(TInPlaceType<TAsyncResolveResult<FResolveVoidResult>>(), MoveTemp(InVoidResult))
+// 	{
+// 	}
+
+// 	FResolveResult(const FResolveResult&) = delete;
+// 	void operator=(const FResolveResult&) = delete;
+
+// 	FResolveResult(FResolveResult&&) = default;
+// 	FResolveResult& operator=(FResolveResult&&) = default;
+
+// 	bool IsVoid() const
+// 	{
+// 		return Value.IsType<TAsyncResolveResult<FResolveVoidResult>>();
+// 	}
+
+// 	bool IsObject() const
+// 	{
+// 		return Value.IsType<TAsyncResolveResult<FResolveObjectResult>>();
+// 	}
+
+// 	bool IsAsync() const
+// 	{
+// 		return IsVoid() ? AsVoid().IsAsync() : AsObject().IsAsync();
+// 	}
+
+
+// 	TAsyncResolveResult<FResolveVoidResult>& AsVoid()
+// 	{
+// 		return Value.Get<TAsyncResolveResult<FResolveVoidResult>>();
+// 	}
+// 	TAsyncResolveResult<FResolveObjectResult>& AsObject()
+// 	{
+// 		return Value.Get<TAsyncResolveResult<FResolveObjectResult>>();
+// 	}
+
+// 	const TAsyncResolveResult<FResolveVoidResult>& AsVoid() const
+// 	{
+// 		return Value.Get<TAsyncResolveResult<FResolveVoidResult>>();
+// 	}
+// 	const TAsyncResolveResult<FResolveObjectResult>& AsObject() const
+// 	{
+// 		return Value.Get<TAsyncResolveResult<FResolveObjectResult>>();
+// 	}
+
+
+// private:
+
+// 	/** Variant value representing either an object, or nothing */
+// 	TVariant<
+// 		TAsyncResolveResult<FResolveVoidResult>,
+// 		TAsyncResolveResult<FResolveObjectResult>
+// 	> Value;
+// };
+
+
+} // namespace UE::UniversalObjectLocator
\ No newline at end of file
diff --git a/Engine/Source/Runtime/UniversalObjectLocator/Public/UniversalObjectLocatorScriptingExtensions.h b/Engine/Source/Runtime/UniversalObjectLocator/Public/UniversalObjectLocatorScriptingExtensions.h
new file mode 100644
--- /dev/null
+++ b/Engine/Source/Runtime/UniversalObjectLocator/Public/UniversalObjectLocatorScriptingExtensions.h
@@ -0,0 +1,72 @@
+// Copyright Epic Games, Inc. All Rights Reserved.
+
+#pragma once
+
+#include "Kismet/BlueprintFunctionLibrary.h"
+#include "Engine/LatentActionManager.h"
+
+#include "UniversalObjectLocatorScriptingExtensions.generated.h"
+
+struct FUniversalObjectLocator;
+
+/**
+ * Function library containing methods that should be hoisted onto FUniversalObjectLocators for scripting
+ */
+UCLASS(MinimalAPI)
+class UUniversalObjectLocatorScriptingExtensions : public UBlueprintFunctionLibrary
+{
+public:
+	GENERATED_BODY()
+
+	/**
+	 * Construct a new universal object locator
+	 */
+	UFUNCTION(BlueprintPure, Category="Universal Object Locators")
+	static FUniversalObjectLocator MakeUniversalObjectLocator(UObject* Object, const UObject* Context);
+
+	/**
+	 * Construct a new universal object locator from a string
+	 */
+	UFUNCTION(BlueprintPure, Category="Universal Object Locators")
+	static FUniversalObjectLocator UniversalObjectLocatorFromString(const FString& InString);
+
+	/**
+	 * Check whether the specified locator is empty; not equivalent to Resolve() != None.
+	 * An empty locator will never resolve to a valid object.
+	 */
+	UFUNCTION(BlueprintPure, Category="Universal Object Locators", meta=(ScriptMethod))
+	static bool IsEmpty(const FUniversalObjectLocator& Locator);
+
+	/**
+	 * Convert the specified locator to its string representation
+	 */
+	UFUNCTION(BlueprintPure, Category="Universal Object Locators", meta=(ScriptMethod))
+	static FString ToString(const FUniversalObjectLocator& Locator);
+
+	/**
+	 * Attempt to resolve the object locator by finding the object. If it is not currently loaded or created, 
+	 * 
+	 * @param Context    (Optional) Context object to use for resolving the object. This should usually be the object that owns or created the locator.
+	 * @return The resolve object pointer, or null if it was not found.
+	 */
+	UFUNCTION(BlueprintPure, Category="Universal Object Locators", meta=(ScriptMethod))
+	static UObject* SyncFind(const FUniversalObjectLocator& Locator, UObject* Context = nullptr);
+
+	/**
+	 * Attempt to resolve the object locator by finding or loading the object.
+	 * 
+	 * @param Context    (Optional) Context object to use for resolving the object. This should usually be the object that owns or created the locator.
+	 * @return The resolve object pointer, or null if it was not found.
+	 */
+	UFUNCTION(BlueprintPure, Category="Universal Object Locators", meta=(ScriptMethod))
+	static UObject* SyncLoad(const FUniversalObjectLocator& Locator, UObject* Context = nullptr);
+
+	/**
+	 * Attempt to resolve the object locator by unloading the object if possible.
+	 * 
+	 * @param Context    (Optional) Context object to use for resolving the object. This should usually be the object that owns or created the locator.
+	 * @return The resolve object pointer, or null if it was not found.
+	 */
+	UFUNCTION(BlueprintPure, Category="Universal Object Locators", meta=(ScriptMethod))
+	static void SyncUnload(const FUniversalObjectLocator& Locator, UObject* Context = nullptr);
+};
\ No newline at end of file
diff --git a/Engine/Source/Runtime/UniversalObjectLocator/Public/UniversalObjectLocatorStringParams.h b/Engine/Source/Runtime/UniversalObjectLocator/Public/UniversalObjectLocatorStringParams.h
new file mode 100644
--- /dev/null
+++ b/Engine/Source/Runtime/UniversalObjectLocator/Public/UniversalObjectLocatorStringParams.h
@@ -0,0 +1,81 @@
+// Copyright Epic Games, Inc. All Rights Reserved.
+
+#pragma once
+
+#include "CoreTypes.h"
+#include "Containers/StringView.h"
+#include "Misc/EnumClassFlags.h"
+#include "Internationalization/Text.h"
+
+namespace UE::UniversalObjectLocator
+{
+
+enum class EParseStringFlags : uint8
+{
+	None = 0,
+
+	/** Whether the user desires error messages to be emited */
+	ErrorMessaging = 1 << 0,
+};
+ENUM_CLASS_FLAGS(EParseStringFlags)
+
+/**
+ * String parameter structure specifying additional information required for a string to be parsed
+ */
+struct FParseStringParams
+{
+	/** Parse flags */
+	EParseStringFlags Flags;
+
+	bool NeedsErrorMessaging() const
+	{
+		return EnumHasAnyFlags(Flags, EParseStringFlags::ErrorMessaging);
+	}
+};
+
+#define UE_UOL_PARSE_ERROR(InParams, ...) InParams.NeedsErrorMessaging() ? (__VA_ARGS__) : FText()
+
+/**
+ * String parse result structure, specfying success/failure, number of characters that were parsed, and
+ *   any error information.
+ */
+struct FParseStringResult
+{
+	/** Optional error message, only populated when FParseStringParams::bWantsErrorMessaging is true */
+	FText ErrorMessage;
+
+	/** The number of characters that were parsed.
+	 *   On Success this indicates the last character that resulted in the successful parse,
+	 *   on failure, this may be 0, or the point at which parsing finished */
+	int32 NumCharsParsed = 0;
+
+	/** Whether parsing was successful or not */
+	bool bSuccess = false;
+
+	explicit operator bool() const
+	{
+		return bSuccess;
+	}
+
+	FParseStringResult& Success(int32 AdditionalNumCharsParsed = 0)
+	{
+		NumCharsParsed += AdditionalNumCharsParsed;
+		bSuccess = true;
+		return *this;
+	}
+
+	FParseStringResult& Failure(const FText& InFailureText)
+	{
+		if (!InFailureText.IsEmpty())
+		{
+			// Only overwrite the error message if it's supplied
+			ErrorMessage = InFailureText;
+		}
+		bSuccess = false;
+		return *this;
+	}
+
+	UNIVERSALOBJECTLOCATOR_API FStringView Progress(FStringView CurrentString, int32 NumToProgress);
+};
+
+} // namespace UE::UniversalObjectLocator
\ No newline at end of file
diff --git a/Engine/Source/Runtime/UniversalObjectLocator/Public/UniversalObjectLocatorStringUtils.h b/Engine/Source/Runtime/UniversalObjectLocator/Public/UniversalObjectLocatorStringUtils.h
new file mode 100644
--- /dev/null
+++ b/Engine/Source/Runtime/UniversalObjectLocator/Public/UniversalObjectLocatorStringUtils.h
@@ -0,0 +1,13 @@
+// Copyright Epic Games, Inc. All Rights Reserved.
+
+#pragma once
+
+#include "CoreTypes.h"
+#include "Containers/StringFwd.h"
+
+namespace UE::UniversalObjectLocator
+{
+
+	UNIVERSALOBJECTLOCATOR_API bool ParseUnsignedInteger(FStringView InString, uint32& OutResult, int32* OutNumCharsParsed = nullptr);
+
+} // namespace UE::UniversalObjectLocator
\ No newline at end of file
diff --git a/Engine/Source/Runtime/UniversalObjectLocator/Public/UniversalObjectLocatorType.h b/Engine/Source/Runtime/UniversalObjectLocator/Public/UniversalObjectLocatorType.h
new file mode 100644
--- /dev/null
+++ b/Engine/Source/Runtime/UniversalObjectLocator/Public/UniversalObjectLocatorType.h
@@ -0,0 +1,19 @@
+// Copyright Epic Games, Inc. All Rights Reserved.
+
+#pragma once
+
+#include "CoreTypes.h"
+
+namespace UE::UniversalObjectLocator
+{
+
+
+enum class ELocatorFragmentType : uint8
+{
+	Absolute,
+	Relative,
+	Undefined,
+};
+
+
+} // namespace UE::UniversalObjectLocator
\ No newline at end of file
diff --git a/Engine/Source/Runtime/UniversalObjectLocator/UniversalObjectLocator.Build.cs b/Engine/Source/Runtime/UniversalObjectLocator/UniversalObjectLocator.Build.cs
new file mode 100644
--- /dev/null
+++ b/Engine/Source/Runtime/UniversalObjectLocator/UniversalObjectLocator.Build.cs
@@ -0,0 +1,19 @@
+// Copyright Epic Games, Inc. All Rights Reserved.
+
+using UnrealBuildTool;
+
+public class UniversalObjectLocator : ModuleRules
+{
+	public UniversalObjectLocator(ReadOnlyTargetRules Target) : base(Target)
+	{
+        UnsafeTypeCastWarningLevel = WarningLevel.Error;
+
+        PublicDependencyModuleNames.AddRange(
+			new string[] {
+				"Core",
+				"CoreUObject",
+				"Engine",
+			}
+		);
+	}
+}
