Miscellaneous compile error fixes

diff --git a/Engine/Source/Runtime/MovieScene/Public/Evaluation/MovieScenePlaybackCapabilities.h b/Engine/Source/Runtime/MovieScene/Public/Evaluation/MovieScenePlaybackCapabilities.h
--- a/Engine/Source/Runtime/MovieScene/Public/Evaluation/MovieScenePlaybackCapabilities.h
+++ b/Engine/Source/Runtime/MovieScene/Public/Evaluation/MovieScenePlaybackCapabilities.h
@@ -236,7 +236,8 @@ protected:
 		if (HasCapability(CapabilityBit))
 		{
 			const int32 Index = GetCapabilityIndex(CapabilityBit);
-			return GetHeader(Index).Resolve(Memory);
+			check(Index >= 0 && Index < 255);
+			return GetHeader(static_cast<uint8>(Index)).Resolve(Memory);
 		}
 
 		return FPlaybackCapabilityPtr();
@@ -245,8 +246,8 @@ protected:
 	FPlaybackCapabilityPtr GetCapabilityChecked(uint32 CapabilityBit) const
 	{
 		const int32 Index = GetCapabilityIndex(CapabilityBit);
-		check(Index != INDEX_NONE);
-		return GetHeader(Index).Resolve(Memory);
+		check(Index >= 0 && Index < 255);
+		return GetHeader(static_cast<uint8>(Index)).Resolve(Memory);
 	}
 
 	/**
diff --git a/Engine/Source/Runtime/UniversalObjectLocator/Private/UniversalObjectLocatorFragment.cpp b/Engine/Source/Runtime/UniversalObjectLocator/Private/UniversalObjectLocatorFragment.cpp
--- a/Engine/Source/Runtime/UniversalObjectLocator/Private/UniversalObjectLocatorFragment.cpp
+++ b/Engine/Source/Runtime/UniversalObjectLocator/Private/UniversalObjectLocatorFragment.cpp
@@ -484,7 +484,7 @@ void FUniversalObjectLocatorFragment::Reset(const UObject* InObject, const UObje
 		FragmentType = MakeFragmentTypeHandle(BestFragmentType);
 		InitializePayload(BestFragmentType->PayloadType);
 
-		BestFragmentType->InitializePayload(GetPayload(), FInitializeParams(InObject, Context));
+		BestFragmentType->InitializePayload(GetPayload(), FInitializeParams{ InObject, Context });
 	}
 }
 
@@ -525,7 +525,7 @@ void FUniversalObjectLocatorFragment::Reset(const UObject* InObject, const UObje
 		FragmentType = MakeFragmentTypeHandle(BestFragmentType);
 		InitializePayload(BestFragmentType->PayloadType);
 
-		BestFragmentType->InitializePayload(GetPayload(), FInitializeParams(InObject, Context));
+		BestFragmentType->InitializePayload(GetPayload(), FInitializeParams{ InObject, Context });
 	}
 }
 
diff --git a/Engine/Source/Runtime/UniversalObjectLocator/Public/UniversalObjectLocatorFragment.h b/Engine/Source/Runtime/UniversalObjectLocator/Public/UniversalObjectLocatorFragment.h
--- a/Engine/Source/Runtime/UniversalObjectLocator/Public/UniversalObjectLocatorFragment.h
+++ b/Engine/Source/Runtime/UniversalObjectLocator/Public/UniversalObjectLocatorFragment.h
@@ -339,7 +339,7 @@ FUniversalObjectLocatorFragment::FUniversalObjectLocatorFragment(UE::UniversalOb
 	}
 
 	// Placement new the payload
-	new (GetPayload()) T(Forward<ArgTypes>(InArgs)...);
+	new (GetPayload()) T{ Forward<ArgTypes>(InArgs)... };
 }
 
 template<typename T>
