Temporarily removed code that causes a cyclic module dependency

diff --git a/Engine/Source/Runtime/UniversalObjectLocator/Private/UniversalObjectLocatorScriptingExtensions.cpp b/Engine/Source/Runtime/UniversalObjectLocator/Private/UniversalObjectLocatorScriptingExtensions.cpp
deleted file mode 100644
diff --git a/Engine/Source/Runtime/UniversalObjectLocator/Public/UniversalObjectLocatorScriptingExtensions.h b/Engine/Source/Runtime/UniversalObjectLocator/Public/UniversalObjectLocatorScriptingExtensions.h
deleted file mode 100644
diff --git a/Engine/Source/Runtime/UniversalObjectLocator/UniversalObjectLocator.Build.cs b/Engine/Source/Runtime/UniversalObjectLocator/UniversalObjectLocator.Build.cs
--- a/Engine/Source/Runtime/UniversalObjectLocator/UniversalObjectLocator.Build.cs
+++ b/Engine/Source/Runtime/UniversalObjectLocator/UniversalObjectLocator.Build.cs
@@ -12,7 +12,6 @@ public class UniversalObjectLocator : ModuleRules
 			new string[] {
 				"Core",
 				"CoreUObject",
-				"Engine",
 			}
 		);
 	}
