Temporarily remove native make/break specification

diff --git a/Engine/Source/Runtime/UniversalObjectLocator/Public/UniversalObjectLocator.h b/Engine/Source/Runtime/UniversalObjectLocator/Public/UniversalObjectLocator.h
--- a/Engine/Source/Runtime/UniversalObjectLocator/Public/UniversalObjectLocator.h
+++ b/Engine/Source/Runtime/UniversalObjectLocator/Public/UniversalObjectLocator.h
@@ -20,7 +20,7 @@
  *
  * This type is 16 bytes.
  */
-USTRUCT(BlueprintType, Category=GameFramework, meta=(HasNativeMake="/Script/UniversalObjectLocator.UniversalObjectLocatorScriptingExtensions.MakeUniversalObjectLocator"))
+USTRUCT(BlueprintType, Category=GameFramework)
 struct FUniversalObjectLocator
 {
 	GENERATED_BODY()
