Fixed relative initialize results being reported as absolute

diff --git a/Engine/Source/Runtime/UniversalObjectLocator/Public/UniversalObjectLocatorInitializeResult.h b/Engine/Source/Runtime/UniversalObjectLocator/Public/UniversalObjectLocatorInitializeResult.h
--- a/Engine/Source/Runtime/UniversalObjectLocator/Public/UniversalObjectLocatorInitializeResult.h
+++ b/Engine/Source/Runtime/UniversalObjectLocator/Public/UniversalObjectLocatorInitializeResult.h
@@ -41,7 +41,7 @@ struct FInitializeResult
 
 		FInitializeResult Result;
 		Result.RelativeToContext = InRelativeToContext;
-		Result.Type = ELocatorType::Absolute;
+		Result.Type = ELocatorType::Relative;
 		return Result;
 	}
 
