Added the ability to provide arbitrary parameters to UOL resolution functions

Parameters must be registered with the UOL module in much the same way as fragment types.
Tidied up globals into a singleton Registry class.
Added a function to FUniversalObjectLocator that allows manual addition of fragments.

Also removed the defunct IActorLocatorFragmentResolver interface in favor of supplying these parameters externally.
Corrected Anim Instance locators being always used when trying to address skeletal mesh components.

diff --git a/Engine/Source/Runtime/Engine/Private/UniversalObjectLocators/ActorLocatorFragment.cpp b/Engine/Source/Runtime/Engine/Private/UniversalObjectLocators/ActorLocatorFragment.cpp
--- a/Engine/Source/Runtime/Engine/Private/UniversalObjectLocators/ActorLocatorFragment.cpp
+++ b/Engine/Source/Runtime/Engine/Private/UniversalObjectLocators/ActorLocatorFragment.cpp
@@ -1,7 +1,6 @@
 // Copyright Epic Games, Inc. All Rights Reserved.
 
 #include "UniversalObjectLocators/ActorLocatorFragment.h"
-#include "UniversalObjectLocators/IActorLocatorFragmentResolver.h"
 #include "UniversalObjectLocatorFragmentTypeHandle.h"
 #include "UniversalObjectLocatorResolveParams.h"
 #include "UniversalObjectLocatorStringParams.h"
@@ -9,13 +8,16 @@
 #include "UniversalObjectLocatorInitializeResult.h"
 
 #include "UObject/UnrealNames.h"
+#include "UObject/Package.h"
 #include "LevelUtils.h"
 #include "Engine/Level.h"
 #include "Engine/LevelStreaming.h"
 #include "Engine/World.h"
+#include "WorldPartition/WorldPartitionLevelHelper.h"
 #include "Misc/EditorPathHelper.h"
 
 UE::UniversalObjectLocator::TFragmentTypeHandle<FActorLocatorFragment> FActorLocatorFragment::FragmentType;
+UE::UniversalObjectLocator::TParameterTypeHandle<FActorLocatorFragmentResolveParameter> FActorLocatorFragmentResolveParameter::ParameterType;
 
 ULevel* GetLevelFromContext(const UObject* InContext)
 {
@@ -75,15 +77,32 @@ UE::UniversalObjectLocator::FResolveResult FActorLocatorFragment::Resolve(const
 
 	UObject* Result = nullptr;
 
-	// Handle the context being a IActorLocatorFragmentResolver object - this is the first port of call
-	if (const IActorLocatorFragmentResolver* Resolver = Cast<const IActorLocatorFragmentResolver>(Params.Context))
+	const FActorLocatorFragmentResolveParameter* Parameter = Params.FindParameter<FActorLocatorFragmentResolveParameter>();
+
+	// If we have a custom fragment parameter specified and it matches the source asset path of this fragment's payload,
+	//   use the streaming world specified by the source asset path.
+	if (Parameter && Parameter->StreamingWorld && Parameter->SourceAssetPath == Path.GetAssetPath())
 	{
-		if (Resolver && Resolver->ResolveActorLocatorPayload(*this, Result))
+		if (!Parameter->ContainerID.IsMainContainer())
+		{
+			// Append the ContainerID and lookup
+			const FString SubPathString = FWorldPartitionLevelHelper::AddActorContainerIDToSubPathString(
+				Parameter->ContainerID, Path.GetSubPathString());
+
+			Parameter->StreamingWorld->ResolveSubobject(*SubPathString, Result, /*bLoadIfExists*/false);
+		}
+		else
 		{
-			return FResolveResultData(Result);
+			// Traditional level streaming needs to resolve bindings from the actual world that owns the streamed level
+			Parameter->StreamingWorld->ResolveSubobject(*Path.GetSubPathString(), Result, /*bLoadIfExists*/false);
 		}
 	}
 
+	if (Result)
+	{
+		return FResolveResultData(Result);
+	}
+
 	// Next handle default level streaming and partition worlds behavior
 	if (ULevel* Level = GetLevelFromContext(Params.Context))
 	{
@@ -93,6 +112,23 @@ UE::UniversalObjectLocator::FResolveResult FActorLocatorFragment::Resolve(const
 	// Finally fallback to just trying to resolve the path directly
 	if (!Result)
 	{
+#if WITH_EDITORONLY_DATA
+		UPackage* ContextPackage = Params.Context ? Params.Context->GetOutermost() : nullptr;
+		if (ContextPackage)
+		{
+			// If we are resolving within a PIE instance, fixp the PIE instance on the path first
+			const int32 PIEInstanceID = ContextPackage->GetPIEInstanceID();
+			if (PIEInstanceID != INDEX_NONE)
+			{
+				FSoftObjectPath PIEPath = Path;
+				PIEPath.FixupForPIE(PIEInstanceID);
+
+				Result = PIEPath.ResolveObject();
+				return FResolveResultData(Result);
+			}
+		}
+#endif
+
 		Result = Path.ResolveObject();
 	}
 
@@ -127,6 +163,24 @@ UE::UniversalObjectLocator::FInitializeResult FActorLocatorFragment::Initialize(
 	Path = InParams.Object;
 #endif
 
+	// Fixup PIE prefixes so that re always reference the non-PIE instance object
+#if WITH_EDITORONLY_DATA
+	UPackage* ObjectPackage = InParams.Object->GetOutermost();
+	if (ensure(ObjectPackage))
+	{
+		// If this is being set from PIE we need to remove the pie prefix and point to the editor object
+		if (ObjectPackage->GetPIEInstanceID() != INDEX_NONE)
+		{
+			TStringBuilder<16> PIEPrefix;
+			PIEPrefix.Appendf(PLAYWORLD_PACKAGE_PREFIX TEXT("_%d_"), ObjectPackage->GetPIEInstanceID());
+
+			FString NewPath = Path.ToString();
+			NewPath.ReplaceInline(*PIEPrefix, TEXT(""));
+			Path.SetPath(NewPath);
+		}
+	}
+#endif
+
 	// Really, actors should be relative to their level in order to support streaming within level instances, but
 	//   world partition makes that impossible
 	return FInitializeResult::Absolute();
diff --git a/Engine/Source/Runtime/Engine/Private/UniversalObjectLocators/AnimInstanceLocatorFragment.cpp b/Engine/Source/Runtime/Engine/Private/UniversalObjectLocators/AnimInstanceLocatorFragment.cpp
--- a/Engine/Source/Runtime/Engine/Private/UniversalObjectLocators/AnimInstanceLocatorFragment.cpp
+++ b/Engine/Source/Runtime/Engine/Private/UniversalObjectLocators/AnimInstanceLocatorFragment.cpp
@@ -23,11 +23,7 @@ UE::UniversalObjectLocator::TFragmentTypeHandle<FAnimInstanceLocatorFragment> FA
 
 uint32 FAnimInstanceLocatorFragment::ComputePriority(const UObject* ObjectToReference, const UObject* Context)
 {
-	if (Context && Context->IsA<USkeletalMeshComponent>())
-	{
-		return 1000;
-	}
-	if (ObjectToReference && (ObjectToReference->IsA<USkeletalMeshComponent>() || ObjectToReference->IsA<UAnimInstance>()))
+	if (ObjectToReference && ObjectToReference->IsA<UAnimInstance>())
 	{
 		return 1000;
 	}
diff --git a/Engine/Source/Runtime/Engine/Private/UnrealEngine.cpp b/Engine/Source/Runtime/Engine/Private/UnrealEngine.cpp
--- a/Engine/Source/Runtime/Engine/Private/UnrealEngine.cpp
+++ b/Engine/Source/Runtime/Engine/Private/UnrealEngine.cpp
@@ -2056,6 +2056,9 @@ void UEngine::Init(IEngineLoop* InEngineLoop)
 			FragmentTypeParams.PrimaryEditorType = "Component";
 			FAnimInstanceLocatorFragment::FragmentType = UolModule.RegisterFragmentType<FAnimInstanceLocatorFragment>(FragmentTypeParams);
 		}
+		{
+			FActorLocatorFragmentResolveParameter::ParameterType = UolModule.RegisterParameterType<FActorLocatorFragmentResolveParameter>();
+		}
 	}
 
 	if (GConfig)
diff --git a/Engine/Source/Runtime/Engine/Public/UniversalObjectLocators/ActorLocatorFragment.h b/Engine/Source/Runtime/Engine/Public/UniversalObjectLocators/ActorLocatorFragment.h
--- a/Engine/Source/Runtime/Engine/Public/UniversalObjectLocators/ActorLocatorFragment.h
+++ b/Engine/Source/Runtime/Engine/Public/UniversalObjectLocators/ActorLocatorFragment.h
@@ -4,8 +4,30 @@
 
 #include "CoreTypes.h"
 #include "UniversalObjectLocatorFwd.h"
+#include "WorldPartition/WorldPartitionActorContainerID.h"
 #include "ActorLocatorFragment.generated.h"
 
+
+/**
+ * 32 Bytes (40 in-editor).
+ */
+USTRUCT()
+struct FActorLocatorFragmentResolveParameter
+{
+	GENERATED_BODY()
+
+	ENGINE_API static UE::UniversalObjectLocator::TParameterTypeHandle<FActorLocatorFragmentResolveParameter> ParameterType;
+
+	UPROPERTY()
+	TObjectPtr<UWorld> StreamingWorld;
+
+	UPROPERTY()
+	FActorContainerID ContainerID;
+
+	UPROPERTY()
+	FTopLevelAssetPath SourceAssetPath;
+};
+
 /**
  * 32 Bytes (40 in-editor).
  */
diff --git a/Engine/Source/Runtime/Engine/Public/UniversalObjectLocators/IActorLocatorFragmentResolver.h b/Engine/Source/Runtime/Engine/Public/UniversalObjectLocators/IActorLocatorFragmentResolver.h
deleted file mode 100644
diff --git a/Engine/Source/Runtime/UniversalObjectLocator/Private/UniversalObjectLocator.cpp b/Engine/Source/Runtime/UniversalObjectLocator/Private/UniversalObjectLocator.cpp
--- a/Engine/Source/Runtime/UniversalObjectLocator/Private/UniversalObjectLocator.cpp
+++ b/Engine/Source/Runtime/UniversalObjectLocator/Private/UniversalObjectLocator.cpp
@@ -6,6 +6,7 @@
 #include "UniversalObjectLocatorInitializeResult.h"
 #include "UniversalObjectLocatorStringParams.h"
 #include "UniversalObjectLocatorStringUtils.h"
+#include "UniversalObjectLocatorRegistry.h"
 
 #define LOCTEXT_NAMESPACE "UOL"
 
@@ -13,8 +14,6 @@ namespace UE::UniversalObjectLocator
 {
 	static constexpr FStringView MagicLeadingString = TEXTVIEW("uobj://");
 
-	extern TArray<FFragmentType> GFragmentTypes;
-
 	const FFragmentType* FindBestFragmentType(const UObject* Object, const UObject* Context);
 
 } // namespace UE::UniversalObjectLocator
@@ -48,6 +47,11 @@ UE::UniversalObjectLocator::FResolveResult FUniversalObjectLocator::Resolve(cons
 {
 	using namespace UE::UniversalObjectLocator;
 
+	if (UE::IsSavingPackage(nullptr) || IsGarbageCollecting())
+	{
+		return FResolveResult();
+	}
+
 	// Check for invalid combinations of flags
 	check(!EnumHasAllFlags(Params.Flags, EResolveFlags::Load | EResolveFlags::Unload));
 	// Cannot have WillWait without Async
@@ -279,6 +283,11 @@ void FUniversalObjectLocator::Reset(UObject* InObject, const UObject* Context, c
 	}
 }
 
+void FUniversalObjectLocator::AddFragment(FUniversalObjectLocatorFragment&& InFragment)
+{
+	Fragments.Emplace(MoveTemp(InFragment));
+}
+
 const UE::UniversalObjectLocator::FFragmentType* FUniversalObjectLocator::GetLastFragmentType() const
 {
 	return Fragments.Num() != 0 ? Fragments.Last().GetFragmentType() : nullptr;
diff --git a/Engine/Source/Runtime/UniversalObjectLocator/Private/UniversalObjectLocatorFragment.cpp b/Engine/Source/Runtime/UniversalObjectLocator/Private/UniversalObjectLocatorFragment.cpp
--- a/Engine/Source/Runtime/UniversalObjectLocator/Private/UniversalObjectLocatorFragment.cpp
+++ b/Engine/Source/Runtime/UniversalObjectLocator/Private/UniversalObjectLocatorFragment.cpp
@@ -6,6 +6,7 @@
 #include "UniversalObjectLocatorStringParams.h"
 #include "UniversalObjectLocatorInitializeParams.h"
 #include "UniversalObjectLocatorInitializeResult.h"
+#include "UniversalObjectLocatorRegistry.h"
 #include "UObject/SoftObjectPath.h"
 #include "Containers/SparseArray.h"
 #include "Misc/AsciiSet.h"
@@ -17,15 +18,13 @@ DEFINE_LOG_CATEGORY(LogUniversalObjectLocator);
 
 namespace UE::UniversalObjectLocator
 {
-	extern TArray<FFragmentType> GFragmentTypes;
-
 	const FFragmentType* FindBestFragmentType(const UObject* Object, const UObject* Context)
 	{
 		// Loop through all our FragmentTypes to find the most supported one
 		uint32 BestFragmentTypePriority = 0;
 		const FFragmentType* BestFragmentType = nullptr;
 
-		for (const FFragmentType& FragmentType : GFragmentTypes)
+		for (const FFragmentType& FragmentType : FRegistry::Get().FragmentTypes)
 		{
 			const uint32 ThisFragmentTypePriority = FragmentType.ComputePriority(Object, Context);
 			if (ThisFragmentTypePriority > BestFragmentTypePriority)
@@ -45,14 +44,14 @@ namespace UE::UniversalObjectLocator
 
 	FFragmentType* FFragmentTypeHandle::Resolve() const
 	{
-		return Handle == 0xff ? nullptr : &GFragmentTypes[Handle];
+		return Handle == 0xff ? nullptr : &FRegistry::Get().FragmentTypes[Handle];
 	}
 
 	FFragmentTypeHandle MakeFragmentTypeHandle(const FFragmentType* FragmentType)
 	{
 		check(FragmentType);
 
-		const uint64 FragmentTypeOffset = static_cast<uint64>(FragmentType - GFragmentTypes.GetData());
+		const uint64 FragmentTypeOffset = static_cast<uint64>(FragmentType - FRegistry::Get().FragmentTypes.GetData());
 		checkf(FragmentTypeOffset < std::numeric_limits<uint8>::max(), TEXT("Maximum number of UOL FragmentTypes reached"));
 
 		return FFragmentTypeHandle(static_cast<uint8>(FragmentTypeOffset));
@@ -331,7 +330,7 @@ UE::UniversalObjectLocator::FParseStringResult FUniversalObjectLocatorFragment::
 	if (FragmentTypeID != NAME_None)
 	{
 		// Find the FragmentType
-		const FFragmentType* SerializedFragmentType = Algo::FindBy(GFragmentTypes, FragmentTypeID, &FFragmentType::FragmentTypeID);
+		const FFragmentType* SerializedFragmentType = FRegistry::Get().FindFragmentType(FragmentTypeID);
 		if (SerializedFragmentType != nullptr && SerializedFragmentType->PayloadType != nullptr)
 		{
 			this->DestroyPayload();
@@ -498,12 +497,13 @@ void FUniversalObjectLocatorFragment::Reset(const UObject* InObject, const UObje
 	uint32 BestFragmentTypePriority = 0;
 	const FFragmentType* BestFragmentType = nullptr;
 
-	if (!ensure(GFragmentTypes.Num() < 255))
+	TArray<FFragmentType>& FragmentTypes = FRegistry::Get().FragmentTypes;
+	if (!ensure(FragmentTypes.Num() < 255))
 	{
 		return;
 	}
 
-	const uint8 Num = static_cast<uint8>(GFragmentTypes.Num());
+	const uint8 Num = static_cast<uint8>(FragmentTypes.Num());
 	for (uint8 Index = 0; Index < Num; ++Index)
 	{
 		if (!CanUseFragmentType(FFragmentTypeHandle(Index)))
@@ -511,7 +511,7 @@ void FUniversalObjectLocatorFragment::Reset(const UObject* InObject, const UObje
 			continue;
 		}
 
-		const FFragmentType& ThisFragmentType = GFragmentTypes[Index];
+		const FFragmentType& ThisFragmentType = FragmentTypes[Index];
 		const uint32 ThisFragmentTypePriority = ThisFragmentType.ComputePriority(InObject, Context);
 		if (ThisFragmentTypePriority > BestFragmentTypePriority)
 		{
@@ -559,7 +559,7 @@ bool FUniversalObjectLocatorFragment::Serialize(FArchive& Ar)
 		else
 		{
 			// Find the FragmentType
-			const FFragmentType* SerializedFragmentType = Algo::FindBy(GFragmentTypes, FragmentTypeID, &FFragmentType::FragmentTypeID);
+			const FFragmentType* SerializedFragmentType = FRegistry::Get().FindFragmentType(FragmentTypeID);
 			if (!SerializedFragmentType || !SerializedFragmentType->PayloadType)
 			{
 				Reset();
diff --git a/Engine/Source/Runtime/UniversalObjectLocator/Private/UniversalObjectLocatorModule.cpp b/Engine/Source/Runtime/UniversalObjectLocator/Private/UniversalObjectLocatorModule.cpp
--- a/Engine/Source/Runtime/UniversalObjectLocator/Private/UniversalObjectLocatorModule.cpp
+++ b/Engine/Source/Runtime/UniversalObjectLocator/Private/UniversalObjectLocatorModule.cpp
@@ -4,6 +4,8 @@
 #include "Modules/ModuleManager.h"
 #include "IUniversalObjectLocatorModule.h"
 #include "UniversalObjectLocatorFragment.h"
+#include "UniversalObjectLocatorRegistry.h"
+#include "UniversalObjectLocatorParameterTypeHandle.h"
 //#include "Modules/VisualizerDebuggingState.h"
 #include "Containers/Ticker.h"
 
@@ -12,8 +14,6 @@
 namespace UE::UniversalObjectLocator
 {
 
-TArray<FFragmentType> GFragmentTypes;
-
 class FUniversalObjectLocatorModule
 	: public IUniversalObjectLocatorModule
 {
@@ -43,25 +43,46 @@ public:
 
 	FFragmentTypeHandle RegisterFragmentTypeImpl(const FFragmentType& FragmentType) override
 	{
-		const int32 Index = GFragmentTypes.Num();
-		GFragmentTypes.Add(FragmentType);
+		FRegistry& Registry = FRegistry::Get();
+
+		const int32 Index = Registry.FragmentTypes.Num();
+		Registry.FragmentTypes.Add(FragmentType);
 		checkf(Index < static_cast<int32>(std::numeric_limits<uint8>::max()), TEXT("Maximum number of UOL FragmentTypes reached"));
 
 		// @todo: enable this code once visualizer debugging state is enabled
 		// Re-assign the debugging ptr in case it changed
-		// UE::Core::FVisualizerDebuggingState::Assign("UOL", GFragmentTypes.GetData());
+		// UE::Core::FVisualizerDebuggingState::Assign("UOL", Registry.FragmentTypes.GetData());
 
 		return FFragmentTypeHandle(static_cast<uint8>(Index));
 	}
 
 	void UnregisterFragmentTypeImpl(FFragmentTypeHandle FragmentType) override
 	{
-		GFragmentTypes[FragmentType.GetIndex()] = FFragmentType{};
+		FRegistry::Get().FragmentTypes[FragmentType.GetIndex()] = FFragmentType{};
+	}
+
+	FParameterTypeHandle RegisterParameterTypeImpl(UScriptStruct* Struct)
+	{
+		FRegistry& Registry = FRegistry::Get();
+
+		const int32 Index = Registry.ParameterTypes.Num();
+		Registry.ParameterTypes.Add(Struct);
+
+		checkf(Index < FResolveParameterBuffer::MaxNumParameters, TEXT("Maximum number of UOL ParameterTypes reached"));
+
+		return FParameterTypeHandle(static_cast<uint8>(Index));
+	}
+
+	void UnregisterParameterTypeImpl(FParameterTypeHandle ParameterType)
+	{
+		FRegistry& Registry = FRegistry::Get();
+		check(ParameterType.IsValid());
+		Registry.ParameterTypes[ParameterType.GetIndex()] = nullptr;
 	}
 
 	bool PurgeVisualizers(float) const
 	{
-		for (FFragmentType& FragmentType : GFragmentTypes)
+		for (FFragmentType& FragmentType : FRegistry::Get().FragmentTypes)
 		{
 			if (FragmentType.DebuggingAssistant)
 			{
diff --git a/Engine/Source/Runtime/UniversalObjectLocator/Private/UniversalObjectLocatorParameterTypeHandle.cpp b/Engine/Source/Runtime/UniversalObjectLocator/Private/UniversalObjectLocatorParameterTypeHandle.cpp
new file mode 100644
--- /dev/null
+++ b/Engine/Source/Runtime/UniversalObjectLocator/Private/UniversalObjectLocatorParameterTypeHandle.cpp
@@ -0,0 +1,23 @@
+// Copyright Epic Games, Inc. All Rights Reserved.
+
+#include "UniversalObjectLocatorParameterTypeHandle.h"
+#include "UniversalObjectLocatorRegistry.h"
+
+namespace UE::UniversalObjectLocator
+{
+
+
+UScriptStruct* FParameterTypeHandle::Resolve(UScriptStruct* Expected) const
+{
+	if (Handle == 0xff)
+	{
+		return nullptr;
+	}
+
+	UScriptStruct* Result = FRegistry::Get().ParameterTypes[Handle];
+	check(!Expected || Result == Expected);
+	return Result;
+}
+
+
+} // namespace UE::UniversalObjectLocator
\ No newline at end of file
diff --git a/Engine/Source/Runtime/UniversalObjectLocator/Private/UniversalObjectLocatorRegistry.cpp b/Engine/Source/Runtime/UniversalObjectLocator/Private/UniversalObjectLocatorRegistry.cpp
new file mode 100644
--- /dev/null
+++ b/Engine/Source/Runtime/UniversalObjectLocator/Private/UniversalObjectLocatorRegistry.cpp
@@ -0,0 +1,38 @@
+// Copyright Epic Games, Inc. All Rights Reserved.
+
+#include "UniversalObjectLocatorRegistry.h"
+#include "UniversalObjectLocatorFwd.h"
+#include "UniversalObjectLocatorFragmentType.h"
+#include "Algo/Find.h"
+
+namespace UE::UniversalObjectLocator
+{
+
+FRegistry& FRegistry::Get()
+{
+	static FRegistry Registry;
+	return Registry;
+}
+
+const FFragmentType* FRegistry::FindFragmentType(FName ID) const
+{
+	return Algo::FindBy(FragmentTypes, ID, &FFragmentType::FragmentTypeID);
+}
+
+void FRegistry::AddReferencedObjects(FReferenceCollector& Collector)
+{
+	Collector.AddReferencedObjects(ParameterTypes);
+	
+	for (FFragmentType& FragmentType : FragmentTypes)
+	{
+		Collector.AddReferencedObject(FragmentType.PayloadType);
+	}
+}
+
+FString FRegistry::GetReferencerName() const
+{
+	return TEXT("UE::UniversalObjectLocator::FRegistry");
+}
+
+
+} // namespace UE::UniversalObjectLocator
\ No newline at end of file
diff --git a/Engine/Source/Runtime/UniversalObjectLocator/Private/UniversalObjectLocatorRegistry.h b/Engine/Source/Runtime/UniversalObjectLocator/Private/UniversalObjectLocatorRegistry.h
new file mode 100644
--- /dev/null
+++ b/Engine/Source/Runtime/UniversalObjectLocator/Private/UniversalObjectLocatorRegistry.h
@@ -0,0 +1,28 @@
+// Copyright Epic Games, Inc. All Rights Reserved.
+
+#pragma once
+
+#include "Modules/ModuleInterface.h"
+#include "UObject/GCObject.h"
+#include "UniversalObjectLocatorFwd.h"
+#include "UniversalObjectLocatorFragmentType.h"
+
+
+namespace UE::UniversalObjectLocator
+{
+
+struct FRegistry
+	: FGCObject
+{
+	static FRegistry& Get();
+
+	const FFragmentType* FindFragmentType(FName ID) const;
+
+	void AddReferencedObjects(FReferenceCollector& Collector) override;
+	FString GetReferencerName() const override;
+
+	TArray<TObjectPtr<UScriptStruct>> ParameterTypes;
+	TArray<FFragmentType> FragmentTypes;
+};
+
+} // namespace UE::UniversalObjectLocator
\ No newline at end of file
diff --git a/Engine/Source/Runtime/UniversalObjectLocator/Private/UniversalObjectLocatorResolveParams.cpp b/Engine/Source/Runtime/UniversalObjectLocator/Private/UniversalObjectLocatorResolveParams.cpp
new file mode 100644
--- /dev/null
+++ b/Engine/Source/Runtime/UniversalObjectLocator/Private/UniversalObjectLocatorResolveParams.cpp
@@ -0,0 +1,55 @@
+// Copyright Epic Games, Inc. All Rights Reserved.
+
+#include "UniversalObjectLocatorResolveParams.h"
+#include "UniversalObjectLocatorFragmentType.h"
+#include "UniversalObjectLocatorFragmentTypeHandle.h"
+
+namespace UE::UniversalObjectLocator
+{
+
+FResolveParameterBuffer::FResolveParameterBuffer()
+	: Memory(nullptr)
+	, AllParameters(0u)
+	, Capacity(0u)
+	, Num(0u)
+{
+}
+
+FResolveParameterBuffer::~FResolveParameterBuffer()
+{
+	Destroy();
+}
+
+void FResolveParameterBuffer::Destroy()
+{
+	check((Memory == nullptr) == (Num == 0));
+
+	if (Memory)
+	{
+		uint8 Index = 0;
+		uint32 RemainingParameters = AllParameters;
+		while (RemainingParameters != 0)
+		{
+			// Get the bit offset of the next reminaing bit
+			const uint8 BitIndex = static_cast<uint8>(FMath::CountTrailingZeros(RemainingParameters));
+			UScriptStruct* Type = FParameterTypeHandle(BitIndex).Resolve();
+
+			void* Ptr = GetHeader(Index).Resolve(Memory);
+			++Index;
+
+			if (ensure(Type))
+			{
+				Type->DestroyStruct(Ptr);
+			}
+
+			RemainingParameters = RemainingParameters & ~(1u << BitIndex);
+		}
+
+		if (bCanFreeMemory)
+		{
+			FMemory::Free(Memory);
+		}
+	}
+}
+
+} // namespace UE::UniversalObjectLocator
\ No newline at end of file
diff --git a/Engine/Source/Runtime/UniversalObjectLocator/Public/IUniversalObjectLocatorModule.h b/Engine/Source/Runtime/UniversalObjectLocator/Public/IUniversalObjectLocatorModule.h
--- a/Engine/Source/Runtime/UniversalObjectLocator/Public/IUniversalObjectLocatorModule.h
+++ b/Engine/Source/Runtime/UniversalObjectLocator/Public/IUniversalObjectLocatorModule.h
@@ -98,6 +98,17 @@ public:
 		return TFragmentTypeHandle<PayloadStructType>(RegisterFragmentTypeImpl(NewFragmentType));
 	}
 
+	template<typename ParameterType>
+	TParameterTypeHandle<ParameterType> RegisterParameterType()
+	{
+		return TParameterTypeHandle<ParameterType>(RegisterParameterTypeImpl(ParameterType::StaticStruct()));
+	}
+
+	void UnregisterParameterType(FParameterTypeHandle InHandle)
+	{
+		UnregisterParameterTypeImpl(InHandle);
+	}
+
 	template<typename PayloadStructType>
 	void UnregisterFragmentType(TFragmentTypeHandle<PayloadStructType> FragmentType)
 	{
@@ -108,6 +119,9 @@ private:
 
 	virtual FFragmentTypeHandle RegisterFragmentTypeImpl(const FFragmentType& FragmentType) = 0;
 	virtual void UnregisterFragmentTypeImpl(FFragmentTypeHandle FragmentType) = 0;
+
+	virtual FParameterTypeHandle RegisterParameterTypeImpl(UScriptStruct* Struct) = 0;
+	virtual void UnregisterParameterTypeImpl(FParameterTypeHandle ParameterType) = 0;
 };
 
 
diff --git a/Engine/Source/Runtime/UniversalObjectLocator/Public/UniversalObjectLocator.h b/Engine/Source/Runtime/UniversalObjectLocator/Public/UniversalObjectLocator.h
--- a/Engine/Source/Runtime/UniversalObjectLocator/Public/UniversalObjectLocator.h
+++ b/Engine/Source/Runtime/UniversalObjectLocator/Public/UniversalObjectLocator.h
@@ -169,6 +169,21 @@ public:
 	 */
 	UNIVERSALOBJECTLOCATOR_API void Reset(UObject* Object, const UObject* Context = nullptr, const UObject* StopAtContext = nullptr);
 
+	/**
+	 * Add a fragment to the end of this locator
+	 *
+	 * @param InFragment       The fragment to add
+	 */
+	UNIVERSALOBJECTLOCATOR_API void AddFragment(FUniversalObjectLocatorFragment&& InFragment);
+
+	/**
+	 * Templated helper for AddFragment
+	 *
+	 * @param InFragment       The fragment to add
+	 */
+	template<typename FragmentType, typename ...ArgTypes>
+	void AddFragment(ArgTypes&&... FragmentArgs);
+
 	/**
 	 * Retrieve the last fragment in this address
 	 */
@@ -217,6 +232,13 @@ private:
 	TArray<FUniversalObjectLocatorFragment> Fragments;
 };
 
+template<typename FragmentType, typename ...ArgTypes>
+void FUniversalObjectLocator::AddFragment(ArgTypes&&... FragmentArgs)
+{
+	TUniversalObjectLocatorFragment<FragmentType> NewFragment(Forward<ArgTypes>(FragmentArgs)...);
+	AddFragment(MoveTemp(NewFragment));
+}
+
 template<>
 struct TStructOpsTypeTraits<FUniversalObjectLocator> : public TStructOpsTypeTraitsBase2<FUniversalObjectLocator>
 {
diff --git a/Engine/Source/Runtime/UniversalObjectLocator/Public/UniversalObjectLocatorFragmentTypeHandle.h b/Engine/Source/Runtime/UniversalObjectLocator/Public/UniversalObjectLocatorFragmentTypeHandle.h
--- a/Engine/Source/Runtime/UniversalObjectLocator/Public/UniversalObjectLocatorFragmentTypeHandle.h
+++ b/Engine/Source/Runtime/UniversalObjectLocator/Public/UniversalObjectLocatorFragmentTypeHandle.h
@@ -71,7 +71,7 @@ struct FFragmentTypeHandle
 	UNIVERSALOBJECTLOCATOR_API FFragmentType* Resolve() const;
 
 private:
-	/** Global index into UE::UniversalObjectLocator::GFragmentTypes */
+	/** Global index into UE::UniversalObjectLocator::FRegistry::Get().FragmentTypes */
 	uint8 Handle;
 };
 
diff --git a/Engine/Source/Runtime/UniversalObjectLocator/Public/UniversalObjectLocatorFwd.h b/Engine/Source/Runtime/UniversalObjectLocator/Public/UniversalObjectLocatorFwd.h
--- a/Engine/Source/Runtime/UniversalObjectLocator/Public/UniversalObjectLocatorFwd.h
+++ b/Engine/Source/Runtime/UniversalObjectLocator/Public/UniversalObjectLocatorFwd.h
@@ -21,6 +21,9 @@ template<typename> struct TFragmentType;
 struct FFragmentTypeHandle;
 template<typename> struct TFragmentTypeHandle;
 
+struct FParameterTypeHandle;
+template<typename> struct TParameterTypeHandle;
+
 struct FInitializeParams;
 struct FInitializeResult;
 struct FResolveParams;
diff --git a/Engine/Source/Runtime/UniversalObjectLocator/Public/UniversalObjectLocatorParameterTypeHandle.h b/Engine/Source/Runtime/UniversalObjectLocator/Public/UniversalObjectLocatorParameterTypeHandle.h
new file mode 100644
--- /dev/null
+++ b/Engine/Source/Runtime/UniversalObjectLocator/Public/UniversalObjectLocatorParameterTypeHandle.h
@@ -0,0 +1,105 @@
+// Copyright Epic Games, Inc. All Rights Reserved.
+
+#pragma once
+
+#include "CoreTypes.h"
+#include "Templates/TypeHash.h"
+#include "UniversalObjectLocatorFwd.h"
+
+namespace UE::UniversalObjectLocator
+{
+
+/**
+ * A handle to a globally registered Universal Object Locator fragment type.
+ *
+ * Registered through IUniversalObjectLocatorModule::RegisterParameterType, fragment types define all the 
+ *   attributes for a specific type of object resolution protocol that can be represented in a UOL.
+ */
+struct FParameterTypeHandle
+{
+	FParameterTypeHandle()
+		: Handle(0xff)
+	{
+	}
+
+	explicit FParameterTypeHandle(uint8 InHandleOffset)
+		: Handle(InHandleOffset)
+	{
+		check(InHandleOffset != 0xff);
+	}
+
+	explicit operator bool() const
+	{
+		return IsValid();
+	}
+
+	bool IsValid() const
+	{
+		return Handle != 0xff;
+	}
+
+	uint8 GetIndex() const
+	{
+		return Handle;
+	}
+
+	friend uint32 GetTypeHash(FParameterTypeHandle In)
+	{
+		return GetTypeHash(In.Handle);
+	}
+
+	friend bool operator==(FParameterTypeHandle A, FParameterTypeHandle B)
+	{
+		return A.Handle == B.Handle;
+	}
+
+	friend bool operator!=(FParameterTypeHandle A, FParameterTypeHandle B)
+	{
+		return A.Handle != B.Handle;
+	}
+
+	friend bool operator<(FParameterTypeHandle A, FParameterTypeHandle B)
+	{
+		return A.Handle < B.Handle;
+	}
+
+	friend bool operator>(FParameterTypeHandle A, FParameterTypeHandle B)
+	{
+		return A.Handle > B.Handle;
+	}
+
+	UNIVERSALOBJECTLOCATOR_API UScriptStruct* Resolve(UScriptStruct* Expected = nullptr) const;
+
+private:
+	/** Global index into UE::UniversalObjectLocator::FRegistry::Get().ParameterTypes */
+	uint8 Handle;
+};
+
+
+/**
+ * Typed version of FParameterTypeHandle that provides typed access to the underlying parameter type
+ */
+template<typename T>
+struct TParameterTypeHandle : FParameterTypeHandle
+{
+	TParameterTypeHandle()
+	{
+	}
+
+	UScriptStruct* Resolve() const
+	{
+		return FParameterTypeHandle::Resolve(T::StaticStruct());
+	}
+
+private:
+
+	friend class IUniversalObjectLocatorModule;
+	friend class FUniversalObjectLocatorModule;
+
+	TParameterTypeHandle(FParameterTypeHandle InHandle)
+		: FParameterTypeHandle(InHandle)
+	{
+	}
+};
+
+} // namespace UE::UniversalObjectLocator
\ No newline at end of file
diff --git a/Engine/Source/Runtime/UniversalObjectLocator/Public/UniversalObjectLocatorResolveParameterBuffer.h b/Engine/Source/Runtime/UniversalObjectLocator/Public/UniversalObjectLocatorResolveParameterBuffer.h
new file mode 100644
--- /dev/null
+++ b/Engine/Source/Runtime/UniversalObjectLocator/Public/UniversalObjectLocatorResolveParameterBuffer.h
@@ -0,0 +1,229 @@
+// Copyright Epic Games, Inc. All Rights Reserved.
+
+#pragma once
+
+#include "CoreTypes.h"
+#include "UniversalObjectLocatorFwd.h"
+#include "UniversalObjectLocatorParameterTypeHandle.h"
+
+class UObject;
+
+namespace UE::UniversalObjectLocator
+{
+
+
+/**
+ * Header that describes an entry in the parameter buffer.
+ */
+struct FResolveParameterHeader
+{
+	// 2 Bytes
+	uint16 ParameterOffset = 0;
+
+	// 1 Byte
+	uint8 Sizeof = 0;
+
+	// 1 Byte
+	uint8 Alignment = 0;
+
+	/** Resolve the given raw pointer into a parameter pointer */
+	void* Resolve(void* InMemory) const
+	{
+		return static_cast<uint8*>(InMemory) + ParameterOffset;
+	}
+};
+
+/**
+ * A buffer of parameter structs that can be optionally supplied to a UOL when resolving to supply additional context.
+ *
+ * Parameter types are registered through IUniversalObjectLocatorModule::RegisterParameterType and are identified by a unique index.
+ * Parameter structs are stored contiguously in memory after an array of headers that specify the offset of each parameter type.
+ *
+ * Checking for the presence of a specific parameter type is very fast, amounting to a a simple bit operation.
+ */
+struct FResolveParameterBuffer
+{
+	/**
+	 * Default constructor
+	 */
+	UNIVERSALOBJECTLOCATOR_API FResolveParameterBuffer();
+
+	/**
+	 * Destructor that calls each parameter's destructor, and frees the memory if necessary
+	 */
+	UNIVERSALOBJECTLOCATOR_API ~FResolveParameterBuffer();
+
+	/**
+	 * Non-copyable
+	 */
+	FResolveParameterBuffer(const FResolveParameterBuffer&) = delete;
+	FResolveParameterBuffer& operator=(const FResolveParameterBuffer&) = delete;
+
+	/**
+	 * Non-moveable
+	 */
+	FResolveParameterBuffer(FResolveParameterBuffer&& RHS) = delete;
+	FResolveParameterBuffer& operator=(FResolveParameterBuffer&& RHS) = delete;
+
+public:
+
+	/**
+	 * Find a parameter of the specified type.
+	 * Relies upon the struct type defining a static TParameterTypeHandle<T> called ParameterType that was previously registered through RegisterParameterType.
+	 *
+	 * @return A pointer to the parameter if present, nullptr otherwise
+	 */
+	template<typename T>
+	const T* FindParameter() const;
+
+
+	/**
+	 * Find a parameter within this container that was added with the specified handle.
+	 *
+	 * @return A pointer to the parameter if present, nullptr otherwise
+	 */
+	template<typename T>
+	const T* FindParameter(TParameterTypeHandle<T> ParameterTypeHandle) const;
+
+
+	/**
+	 * Add a new parameter to this container using its preregistered type handle.
+	 * This function is accepts 0 or more additional construction arguments which are passed to the parameter's constructor.
+	 *
+	 * @note: The definition of this function is located in UniversalObjectLocatorResolveParameterBuffer.inl
+	 *
+	 * @return A pointer to the parameter if present, nullptr otherwise
+	 */
+	template<typename T, typename ...ArgTypes>
+	T* AddParameter(TParameterTypeHandle<T> ParameterTypeHandle, ArgTypes&&... InArgs);
+
+protected:
+
+	/**
+	 * Check whether the parameter for the specified parameter bit is present
+	 */
+	bool HasParameter(uint32 ParameterBit) const;
+
+	/**
+	 * Implementation function for finding a parameter ptr by its bit
+	 */
+	const void* FindParameterImpl(uint32 ParameterBit) const;
+
+
+	/**
+	 * Creates and stores a new Parameter object at the given bit.
+	 * 
+	 * @note: The definition of this function is located in UniversalObjectLocatorResolveParameterBuffer.inl
+	 */
+	template<typename ParameterType, typename ...ArgTypes>
+	ParameterType* AddParameterImpl(uint32 ParameterBit, ArgTypes&&... InArgs);
+
+
+	/**
+	 * Retrieve the header at the specified index
+	 */
+	const FResolveParameterHeader& GetHeader(uint8 Index) const;
+
+
+	/**
+	 * Compute the index of the parameter that relates to the specified bit
+	 */
+	int32 GetParameterIndex(uint32 ParameterBit) const;
+
+
+	/**
+	 * Destroy this buffer and all its parameters
+	 */
+	void Destroy();
+
+protected:
+
+	/** Memory allocation of the layout [header_1|...|header_n|entry_0|...|entry_n] */
+	uint8* Memory = nullptr;
+
+	/** Bitmask of all the parameters contained within this container */
+	uint32 AllParameters = 0u;
+
+	/** Current capacity of Memory in bytes */
+	uint16 Capacity = 0u;
+
+	/** The number of parameters (and thus, headers) in this container */
+	uint8 Num = 0u;
+
+	/** Boolean that indicates whether we own the memory pointed to by Memory, or if it was supplied externally */
+	bool bCanFreeMemory = false;
+
+public:
+
+	static constexpr int32 MaxNumParameters = sizeof(AllParameters) * 8;
+};
+
+/**
+ * A FResolveParameterBuffer that has an inline memory capacity of the specified size
+ */
+template<int InlineByteSize>
+struct alignas(8) TInlineResolveParameterBuffer : FResolveParameterBuffer
+{
+	TInlineResolveParameterBuffer()
+	{
+		Memory = InlineBuffer;
+		Capacity = InlineByteSize;
+
+		bCanFreeMemory = false;
+	}
+
+private:
+
+	/** Inline memory buffer */
+	uint8 InlineBuffer[InlineByteSize];
+};
+
+
+template<typename T>
+const T* FResolveParameterBuffer::FindParameter() const
+{
+	return FindParameter(T::ParameterType);
+}
+
+template<typename T>
+const T* FResolveParameterBuffer::FindParameter(TParameterTypeHandle<T> ParameterTypeHandle) const
+{
+	const uint8 ParameterIndex = ParameterTypeHandle.GetIndex();
+
+	const uint32 ParameterBit = 1 << ParameterIndex;
+	return static_cast<const T*>(FindParameterImpl(ParameterBit));
+}
+
+inline bool FResolveParameterBuffer::HasParameter(uint32 ParameterBit) const
+{
+	return (AllParameters & ParameterBit) != 0;
+}
+
+inline const void* FResolveParameterBuffer::FindParameterImpl(uint32 ParameterBit) const
+{
+	if (HasParameter(ParameterBit))
+	{
+		const int32 Index = GetParameterIndex(ParameterBit);
+		check(Index >= 0 && Index < 255);
+		return GetHeader(static_cast<uint8>(Index)).Resolve(Memory);
+	}
+
+	return nullptr;
+}
+
+inline const FResolveParameterHeader& FResolveParameterBuffer::GetHeader(uint8 Index) const
+{
+	check(Index < Num);
+	return reinterpret_cast<FResolveParameterHeader*>(Memory)[Index];
+}
+
+inline int32 FResolveParameterBuffer::GetParameterIndex(uint32 ParameterBit) const
+{
+	if ( (AllParameters & ParameterBit) == 0)
+	{
+		return INDEX_NONE;
+	}
+	return FMath::CountBits(AllParameters & (ParameterBit-1u));
+}
+
+} // namespace UE::UniversalObjectLocator
\ No newline at end of file
diff --git a/Engine/Source/Runtime/UniversalObjectLocator/Public/UniversalObjectLocatorResolveParameterBuffer.inl b/Engine/Source/Runtime/UniversalObjectLocator/Public/UniversalObjectLocatorResolveParameterBuffer.inl
new file mode 100644
--- /dev/null
+++ b/Engine/Source/Runtime/UniversalObjectLocator/Public/UniversalObjectLocatorResolveParameterBuffer.inl
@@ -0,0 +1,187 @@
+// Copyright Epic Games, Inc. All Rights Reserved.
+
+#pragma once
+
+#include "UniversalObjectLocatorResolveParams.h"
+
+namespace UE::UniversalObjectLocator
+{
+
+
+template<typename T, typename ...ArgTypes>
+T* FResolveParameterBuffer::AddParameter(TParameterTypeHandle<T> ParameterTypeHandle, ArgTypes&&... InArgs)
+{
+	const UScriptStruct* ParameterStruct = ParameterTypeHandle.Resolve();
+	if (!ensureMsgf(ParameterStruct, TEXT("Parameter struct is no longer registered!")))
+	{
+		return nullptr;
+	}
+
+	const uint8 ParameterIndex = ParameterTypeHandle.GetIndex();
+
+	const uint32 ParameterBit = 1 << ParameterIndex;
+	return static_cast<T*>(AddParameterImpl<T>(ParameterBit, Forward<ArgTypes>(InArgs)...));
+}
+
+
+template<typename ParameterType, typename ...ArgTypes>
+ParameterType* FResolveParameterBuffer::AddParameterImpl(uint32 ParameterBit, ArgTypes&&... InArgs)
+{
+	checkf((AllParameters & ParameterBit) == 0, TEXT("Parameter already exists!"));
+
+	// Add the enum entry
+	AllParameters |= ParameterBit;
+
+	// Find the index of the new Parameter by counting how many bits are set before it
+	// For example, given ParameterBit=0b00010000 and AllParameters=0b00011011:
+	//                       (ParameterBit-1) = 0b00001111
+	//       AllParameters & (ParameterBit-1) = 0b00001011
+	//                  CountBits(0b00001011) = 3
+	const int32 NewParameterIndex = static_cast<int32>(FMath::CountBits(AllParameters & (ParameterBit-1u)));
+
+	FResolveParameterHeader* ExistingHeaders = reinterpret_cast<FResolveParameterHeader*>(Memory);
+
+	uint64 RequiredAlignment = FMath::Max(alignof(ParameterType), alignof(FResolveParameterHeader));
+
+	const int32 ExistingNum = static_cast<int32>(Num);
+
+	// Compute our required alignment for the allocation
+	{
+		for (int32 Index = 0; Index < ExistingNum; ++Index)
+		{
+			RequiredAlignment = FMath::Max(RequiredAlignment, (uint64)ExistingHeaders[Index].Alignment);
+		}
+	}
+
+	uint64 RequiredSizeof = 0u;
+
+	// Compute the required size of our allocation
+	{
+		// Allocate space for headers
+		RequiredSizeof += (ExistingNum+1) * sizeof(FResolveParameterHeader);
+
+		int32 Index = 0;
+
+		// Count up the sizes and alignments of pre-existing parameters that exist before this new entry
+		for (; Index < NewParameterIndex; ++Index)
+		{
+			RequiredSizeof = Align(RequiredSizeof, (uint64)ExistingHeaders[Index].Alignment);
+			RequiredSizeof += ExistingHeaders[Index].Sizeof;
+		}
+
+		// Count up the size and alignment for the new parameter
+		RequiredSizeof = Align(RequiredSizeof, alignof(ParameterType));
+		RequiredSizeof += sizeof(ParameterType);
+		++Index;
+
+		// Now count up the sizes and alignments of pre-existing parameters that exist after this new entry
+		for (; Index < ExistingNum+1; ++Index)
+		{
+			RequiredSizeof = Align(RequiredSizeof, (uint64)ExistingHeaders[Index-1].Alignment);
+			RequiredSizeof += ExistingHeaders[Index-1].Sizeof;
+		}
+	}
+
+	check( RequiredAlignment <= 0XFF );
+
+	/// ----------------------------------------
+
+	uint8* OldAllocation = Memory;
+
+	const bool bShouldFreeMemory = bCanFreeMemory;
+
+	// Make a new allocation if necessary
+	const bool bNeedsReallocation = !IsAligned(Memory, RequiredAlignment) || RequiredSizeof > Capacity;
+	if (bNeedsReallocation)
+	{
+		const uint64 RequiredCapacity = FMath::Max(RequiredSizeof, uint64(Capacity)*2);
+		check(RequiredCapacity <= std::numeric_limits<uint16>::max());
+
+		// Use the greater of the required size or double the current size to allow some additional capcity
+		Capacity = static_cast<uint16>(FMath::Max(RequiredSizeof, uint64(Capacity)*2));
+		Memory = reinterpret_cast<uint8*>(FMemory::Malloc(Capacity, RequiredAlignment));
+
+		bCanFreeMemory = true;
+	}
+
+	// We now have an extra entry
+	++Num;
+
+	uint8* ParameterPtr = Memory + RequiredSizeof;
+
+	auto RelocateParameter = [this, ExistingHeaders, OldAllocation, &ParameterPtr](int32 OldIndex)
+	{
+		// Go back
+		ParameterPtr -= ExistingHeaders[OldIndex].Sizeof;
+		ParameterPtr = AlignDown(ParameterPtr, ExistingHeaders[OldIndex].Alignment);
+
+		void* OldParameter = ExistingHeaders[OldIndex].Resolve(OldAllocation);
+		FMemory::Memmove(ParameterPtr, OldParameter, ExistingHeaders[OldIndex].Sizeof);
+
+		// Overwrite the old header with its new position for error checking when we write the new header
+		ExistingHeaders[OldIndex].ParameterOffset = static_cast<uint16>(ParameterPtr-Memory);
+	};
+
+	// Relocate old structs that proceed the new index
+	for (int32 Index = static_cast<int32>(Num) - 1; Index > NewParameterIndex; --Index)
+	{
+		RelocateParameter(Index-1);
+	}
+
+	// Make the new entry
+	{
+		ParameterPtr -= sizeof(ParameterType);
+		ParameterPtr = AlignDown(ParameterPtr, alignof(ParameterType));
+
+		ParameterType* NewParameterPtr = reinterpret_cast<ParameterType*>(ParameterPtr);
+
+		// Allocate the new type
+		new (NewParameterPtr) ParameterType { Forward<ArgTypes>(InArgs)... };
+
+		static_assert(alignof(ParameterType) < 0x7F, "Required alignment of parameter must fit in 7 bytes");
+	}
+
+	FResolveParameterHeader NewHeader = {
+		static_cast<uint16>(ParameterPtr-Memory),
+		sizeof(ParameterType),
+		alignof(ParameterType)
+	};
+
+	// Relocate the entries that are before the new one
+	for (int32 Index = NewParameterIndex-1; Index >= 0; --Index)
+	{
+		RelocateParameter(Index);
+	}
+
+	// Check that we have enough space before for the headers (this should always be the case)
+	check(ParameterPtr >= Memory + sizeof(FResolveParameterHeader)*Num);
+
+	FResolveParameterHeader* NewHeaders = reinterpret_cast<FResolveParameterHeader*>(Memory);
+
+	for (int32 Index = static_cast<int32>(Num)-1; Index > NewParameterIndex; --Index)
+	{
+		new (&NewHeaders[Index]) FResolveParameterHeader(ExistingHeaders[Index-1]);
+	};
+
+	{
+		new (&NewHeaders[NewParameterIndex]) FResolveParameterHeader(NewHeader);
+	}
+
+	if (Memory != OldAllocation)
+	{
+		for (int32 Index = NewParameterIndex-1; Index >= 0; --Index)
+		{
+			new (&NewHeaders[Index]) FResolveParameterHeader(ExistingHeaders[Index]);
+		};
+	}
+
+	// Tidy up the old allocation. We do not call destructors here because we relocated everything.
+	if (bNeedsReallocation && OldAllocation && bShouldFreeMemory)
+	{
+		FMemory::Free(OldAllocation);
+	}
+
+	return static_cast<ParameterType*>(NewHeaders[NewParameterIndex].Resolve(Memory));
+}
+
+} // namespace UE::UniversalObjectLocator
\ No newline at end of file
diff --git a/Engine/Source/Runtime/UniversalObjectLocator/Public/UniversalObjectLocatorResolveParams.h b/Engine/Source/Runtime/UniversalObjectLocator/Public/UniversalObjectLocatorResolveParams.h
--- a/Engine/Source/Runtime/UniversalObjectLocator/Public/UniversalObjectLocatorResolveParams.h
+++ b/Engine/Source/Runtime/UniversalObjectLocator/Public/UniversalObjectLocatorResolveParams.h
@@ -5,6 +5,8 @@
 #include "CoreTypes.h"
 #include "Async/Future.h"
 #include "Misc/TVariant.h"
+#include "UniversalObjectLocatorFwd.h"
+#include "UniversalObjectLocatorResolveParameterBuffer.h"
 
 class UObject;
 
@@ -37,7 +39,6 @@ enum class EResolveFlags : uint8
 };
 ENUM_CLASS_FLAGS(EResolveFlags)
 
-
 /**
  * Parameters required to resolve a universal object locator
  */
@@ -61,6 +62,18 @@ struct FResolveParams
 	{
 	}
 
+	template<typename T>
+	const T* FindParameter() const
+	{
+		return FindParameter(T::ParameterType);
+	}
+
+	template<typename T>
+	const T* FindParameter(TParameterTypeHandle<T> ParameterType) const
+	{
+		return ParameterBuffer ? ParameterBuffer->FindParameter(ParameterType) : nullptr;
+	}
+
 	/** 
 	 * Utility function to indicate that the object should be found, but not loaded or created if it is not currently
 	 */
@@ -112,10 +125,36 @@ struct FResolveParams
 	/** (Optional) Object to use as a context for resolution. Normally this is the object that owns the reference being resolved */
 	const UObject* Context;
 
+	/** (Optional) Resolve buffer */
+	const FResolveParameterBuffer* ParameterBuffer;
+
 	/** Flag structure */
 	EResolveFlags Flags;
 };
 
+template<int InlineSize>
+struct TResolveParamsWithBuffer : TInlineResolveParameterBuffer<InlineSize>, FResolveParams
+{
+	TResolveParamsWithBuffer()
+	{
+		ParameterBuffer = this;
+	}
+
+	TResolveParamsWithBuffer(const UObject* InContext)
+		: FResolveParams(InContext)
+	{
+		ParameterBuffer = this;
+	}
+
+	TResolveParamsWithBuffer(const UObject* InContext, EResolveFlags InFlags)
+		: FResolveParams(InContext, InFlags)
+	{
+		ParameterBuffer = this;
+	}
+};
+
+
+
 
 /**
  * Flag structure returned from an attempt to resolve a Universal Object Locator
@@ -291,128 +330,4 @@ private:
 	> Value;
 };
 
-// struct FResolveObjectResult
-// {
-// 	FResolveObjectResult()
-// 		: Object(nullptr)
-// 	{
-// 	}
-
-// 	FResolveObjectResult(UObject* InObject, FResolveResultFlags InFlags = FResolveResultFlags())
-// 		: Object(InObject)
-// 		, Flags(InFlags)
-// 	{
-// 	}
-
-// 	/**
-// 	 * The resulting object or nullptr if it could not be found or loaded per the request params.
-// 	 * Should only be set if Params.Flags.bAsync is false.
-// 	 */
-// 	UObject* Object = nullptr;
-
-// 	/** Flags relating to the operation */
-// 	FResolveResultFlags Flags;
-// };
-
-
-// struct FResolveVoidResult
-// {
-// };
-
-// /**
-//  * Result structure that is returned from a request to resolve a Universal Object Locator
-//  */
-// struct FResolveResult
-// {
-// 	/**
-// 	 * Empty result
-// 	 */
-// 	FResolveResult()
-// 		: Value(TInPlaceType<TAsyncResolveResult<FResolveVoidResult>>())
-// 	{
-// 	}
-
-// 	/**
-// 	 * Constructor to use when an object needs to be returned from a Resolve function (ie, Load or Find)
-// 	 */
-// 	FResolveResult(FResolveObjectResult&& InResolveResult)
-// 		: Value(TInPlaceType<TAsyncResolveResult<FResolveObjectResult>>(), MoveTemp(InResolveResult))
-// 	{
-// 	}
-
-// 	/**
-// 	 * Constructor to use when no information needs to be returned from a Resolve function (ie, Unload)
-// 	 */
-// 	FResolveResult(FResolveVoidResult&& InVoidResult)
-// 		: Value(TInPlaceType<TAsyncResolveResult<FResolveVoidResult>>(), MoveTemp(InVoidResult))
-// 	{
-// 	}
-
-// 	/**
-// 	 * Constructor to use when an object needs to be returned from a Resolve function (ie, Load or Find)
-// 	 */
-// 	FResolveResult(TAsyncResolveResult<FResolveObjectResult>&& InResolveResult)
-// 		: Value(TInPlaceType<TAsyncResolveResult<FResolveObjectResult>>(), MoveTemp(InResolveResult))
-// 	{
-// 	}
-
-// 	/**
-// 	 * Constructor to use when no information needs to be returned from a Resolve function (ie, Unload)
-// 	 */
-// 	FResolveResult(TAsyncResolveResult<FResolveVoidResult>&& InVoidResult)
-// 		: Value(TInPlaceType<TAsyncResolveResult<FResolveVoidResult>>(), MoveTemp(InVoidResult))
-// 	{
-// 	}
-
-// 	FResolveResult(const FResolveResult&) = delete;
-// 	void operator=(const FResolveResult&) = delete;
-
-// 	FResolveResult(FResolveResult&&) = default;
-// 	FResolveResult& operator=(FResolveResult&&) = default;
-
-// 	bool IsVoid() const
-// 	{
-// 		return Value.IsType<TAsyncResolveResult<FResolveVoidResult>>();
-// 	}
-
-// 	bool IsObject() const
-// 	{
-// 		return Value.IsType<TAsyncResolveResult<FResolveObjectResult>>();
-// 	}
-
-// 	bool IsAsync() const
-// 	{
-// 		return IsVoid() ? AsVoid().IsAsync() : AsObject().IsAsync();
-// 	}
-
-
-// 	TAsyncResolveResult<FResolveVoidResult>& AsVoid()
-// 	{
-// 		return Value.Get<TAsyncResolveResult<FResolveVoidResult>>();
-// 	}
-// 	TAsyncResolveResult<FResolveObjectResult>& AsObject()
-// 	{
-// 		return Value.Get<TAsyncResolveResult<FResolveObjectResult>>();
-// 	}
-
-// 	const TAsyncResolveResult<FResolveVoidResult>& AsVoid() const
-// 	{
-// 		return Value.Get<TAsyncResolveResult<FResolveVoidResult>>();
-// 	}
-// 	const TAsyncResolveResult<FResolveObjectResult>& AsObject() const
-// 	{
-// 		return Value.Get<TAsyncResolveResult<FResolveObjectResult>>();
-// 	}
-
-
-// private:
-
-// 	/** Variant value representing either an object, or nothing */
-// 	TVariant<
-// 		TAsyncResolveResult<FResolveVoidResult>,
-// 		TAsyncResolveResult<FResolveObjectResult>
-// 	> Value;
-// };
-
-
 } // namespace UE::UniversalObjectLocator
\ No newline at end of file
