Sequencer: Fixed uninitialized member variable

diff --git a/Engine/Source/Runtime/UniversalObjectLocator/Public/UniversalObjectLocatorResolveParams.h b/Engine/Source/Runtime/UniversalObjectLocator/Public/UniversalObjectLocatorResolveParams.h
--- a/Engine/Source/Runtime/UniversalObjectLocator/Public/UniversalObjectLocatorResolveParams.h
+++ b/Engine/Source/Runtime/UniversalObjectLocator/Public/UniversalObjectLocatorResolveParams.h
@@ -46,18 +46,21 @@ struct FResolveParams
 {
 	FResolveParams()
 		: Context(nullptr)
+		, ParameterBuffer(nullptr)
 		, Flags(EResolveFlags::None)
 	{
 	}
 
 	FResolveParams(const UObject* InContext)
 		: Context(InContext)
+		, ParameterBuffer(nullptr)
 		, Flags(EResolveFlags::None)
 	{
 	}
 
 	FResolveParams(const UObject* InContext, EResolveFlags InFlags)
 		: Context(InContext)
+		, ParameterBuffer(nullptr)
 		, Flags(InFlags)
 	{
 	}
