Corrected list of valid payload characters for UOL fragments

diff --git a/Engine/Source/Runtime/UniversalObjectLocator/Public/UniversalObjectLocatorFragment.h b/Engine/Source/Runtime/UniversalObjectLocator/Public/UniversalObjectLocatorFragment.h
--- a/Engine/Source/Runtime/UniversalObjectLocator/Public/UniversalObjectLocatorFragment.h
+++ b/Engine/Source/Runtime/UniversalObjectLocator/Public/UniversalObjectLocatorFragment.h
@@ -38,8 +38,8 @@ struct alignas(8) FUniversalObjectLocatorFragment
 	using FParseStringParams = UE::UniversalObjectLocator::FParseStringParams;
 
 	static constexpr FAsciiSet ValidFragmentTypeCharacters = "abcdefghijklmnopqrstuvwxyzABCDEFGHIJKLMNOPQRSTUVWXYZ0123456789_-";
-	// Valid fragment delimiters are those that are allowable by RFC3986 including '/' and excluding '&' which we use to separate fragments within a path
-	static constexpr FAsciiSet ValidFragmentDelimiters = "!&'()*+,;=";
+	// Valid fragment delimiters are those that are allowable by RFC3986 for the query part (unreserved / pct-encoded / sub-delims / ":" / "@" / "/" / "?") excluding '&' which we use to separate fragments within a path
+	static constexpr FAsciiSet ValidFragmentDelimiters = "%!$'()*+,;=/?:@.~";
 	static constexpr FAsciiSet ValidFragmentPayloadCharacters = ValidFragmentTypeCharacters | ValidFragmentDelimiters;
 
 	/** Make our inline data buffer larger in-editor to support editor-only data without allocation */
