Fix UOL ImportTextItem not incrementing past the parsed buffer on success

diff --git a/Engine/Source/Runtime/UniversalObjectLocator/Private/UniversalObjectLocator.cpp b/Engine/Source/Runtime/UniversalObjectLocator/Private/UniversalObjectLocator.cpp
--- a/Engine/Source/Runtime/UniversalObjectLocator/Private/UniversalObjectLocator.cpp
+++ b/Engine/Source/Runtime/UniversalObjectLocator/Private/UniversalObjectLocator.cpp
@@ -687,6 +687,8 @@ bool FUniversalObjectLocator::ImportTextItem(const TCHAR*& Buffer, int32 PortFla
 			FStringView View(Buffer + 1, int32(BufferEnd - Buffer) - 1);
 			if (TryParseString(View, FParseStringParams()))
 			{
+				// ImportText parsing requires that we increment the buffer ptr beyond the closing ')' on success
+				Buffer = BufferEnd + 1;
 				return true;
 			}
 		}
