Fix uninitialized variable in UOL string parsing params

diff --git a/Engine/Source/Runtime/UniversalObjectLocator/Public/UniversalObjectLocatorStringParams.h b/Engine/Source/Runtime/UniversalObjectLocator/Public/UniversalObjectLocatorStringParams.h
--- a/Engine/Source/Runtime/UniversalObjectLocator/Public/UniversalObjectLocatorStringParams.h
+++ b/Engine/Source/Runtime/UniversalObjectLocator/Public/UniversalObjectLocatorStringParams.h
@@ -25,7 +25,7 @@ ENUM_CLASS_FLAGS(EParseStringFlags)
 struct FParseStringParams
 {
 	/** Parse flags */
-	EParseStringFlags Flags;
+	EParseStringFlags Flags = EParseStringFlags::None;
 
 	bool NeedsErrorMessaging() const
 	{
