Sequencer: Initialize Flags in FFragmentTypeParameters

diff --git a/Engine/Source/Runtime/UniversalObjectLocator/Public/IUniversalObjectLocatorModule.h b/Engine/Source/Runtime/UniversalObjectLocator/Public/IUniversalObjectLocatorModule.h
--- a/Engine/Source/Runtime/UniversalObjectLocator/Public/IUniversalObjectLocatorModule.h
+++ b/Engine/Source/Runtime/UniversalObjectLocator/Public/IUniversalObjectLocatorModule.h
@@ -24,6 +24,7 @@ struct FFragmentTypeParameters
 	FFragmentTypeParameters(FName InFragmentTypeID, FText InDisplayText)
 		: DisplayText(InDisplayText)
 		, FragmentTypeID(InFragmentTypeID)
+		, Flags(EFragmentTypeFlags::None)
 	{
 	}
 
