Improved visualizers for UOL fragments

Previously the visualizers used a method that had side-effects that required manual re-evaluation by the user. These have now been updated to use a side-effects-free method by emplacing a polymorphic vfptr in front of the fragment, along with the newly introduced FVisualizerDebuggingState to provide a much simpler and more intuitive debugging view for UOLs.

diff --git a/Engine/Source/Runtime/UniversalObjectLocator/Private/UniversalObjectLocatorFragment.cpp b/Engine/Source/Runtime/UniversalObjectLocator/Private/UniversalObjectLocatorFragment.cpp
--- a/Engine/Source/Runtime/UniversalObjectLocator/Private/UniversalObjectLocatorFragment.cpp
+++ b/Engine/Source/Runtime/UniversalObjectLocator/Private/UniversalObjectLocatorFragment.cpp
@@ -6,10 +6,12 @@
 #include "UniversalObjectLocatorStringParams.h"
 #include "UniversalObjectLocatorInitializeParams.h"
 #include "UniversalObjectLocatorInitializeResult.h"
+#include "UniversalObjectLocatorFragmentDebugging.h"
 #include "UniversalObjectLocatorRegistry.h"
 #include "UObject/SoftObjectPath.h"
 #include "Containers/SparseArray.h"
 #include "Misc/AsciiSet.h"
+#include "Templates/AlignmentTemplates.h"
 
 #define LOCTEXT_NAMESPACE "UOL"
 
@@ -57,12 +59,46 @@ namespace UE::UniversalObjectLocator
 		return FFragmentTypeHandle(static_cast<uint8>(FragmentTypeOffset));
 	}
 
+	/**
+	 * Compute the size required for the debug header of a fragment with a certain size and alignment constraint
+	 * 
+	 * Since our byte array is explicitly aligned to 8 bytes, we can insert the debug header right at the start of
+	 * our bytes without changing the alignment of the proceeding type.
+	 * 
+	 * If however our type's requested alignment is greater, we use the alignment itself and allocate the header at
+	 *   the tail of that space. For instance, for a 16 byte aligned payload:
+	 * 0..				8..						16..													16 + sizeof(T)
+	 * [				TFragmentPayload<T>		| T Payload 											]
+	 *
+	 * For a 32 byte aligned payload:
+	 * 0..												24..					32..													32 + sizeof(T)
+	 * [												TFragmentPayload<T>		| T Payload 											]
+	 * 
+	 */
+	uint8 ComputeDebugHeaderLog2(size_t Alignment)
+	{
+#if UE_UNIVERSALOBJECTLOCATOR_DEBUG
+		static_assert(alignof(IFragmentPayload) == 8 && sizeof(IFragmentPayload) == 8, "Unexpected alignment/size of IFragmentPayload!");
+
+		const uint32 HeaderCapacityLog2 = FMath::CeilLogTwo((uint32)FMath::Max(Alignment, sizeof(IFragmentPayload)));
+
+		// This value is stored in 6 bits of a uint8
+		// We should never encounter a type aligned > 2^63!
+		check(HeaderCapacityLog2 <= 63);
+
+		return static_cast<uint8>(HeaderCapacityLog2);
+#else
+		return 0;
+#endif
+	}
+
 } // UE::UniversalObjectLocator
 
 
 FUniversalObjectLocatorFragment::FUniversalObjectLocatorFragment(const UObject* InObject, UObject* Context)
 	: bIsInitialized(0)
 	, bIsInline(0)
+	, DebugHeaderSizeLog2(0)
 {
 	Reset(InObject, Context);
 }
@@ -71,15 +107,17 @@ FUniversalObjectLocatorFragment::FUniversalObjectLocatorFragment(const UE::Unive
 	: FragmentType(MakeFragmentTypeHandle(&InFragmentType))
 	, bIsInitialized(0)
 	, bIsInline(0)
+	, DebugHeaderSizeLog2(0)
 {
-	this->InitializePayload(InFragmentType.PayloadType);
+	this->DefaultConstructPayload(InFragmentType);
 }
 
 FUniversalObjectLocatorFragment::FUniversalObjectLocatorFragment()
 	: bIsInitialized(0)
 	, bIsInline(0)
+	, DebugHeaderSizeLog2(0)
 {
-	static_assert(sizeof(FUniversalObjectLocatorFragment) == FUniversalObjectLocatorFragment::Size, "Unexpected size for FUniversalObjectLocatorFragment");
+	static_assert(sizeof(FUniversalObjectLocatorFragment) == FUniversalObjectLocatorFragment::SizeInMemory, "Unexpected size for FUniversalObjectLocatorFragment");
 	static_assert(offsetof(FUniversalObjectLocatorFragment, Data) == 0, "FUniversalObjectLocatorFragment inline data is not aligned properly");
 }
 
@@ -95,17 +133,18 @@ FUniversalObjectLocatorFragment::FUniversalObjectLocatorFragment(const FUniversa
 	: FragmentType(RHS.FragmentType)
 	, bIsInitialized(0)
 	, bIsInline(0)
+	, DebugHeaderSizeLog2(0)
 {
 	using namespace UE::UniversalObjectLocator;
 
 	if (RHS.bIsInitialized)
 	{
-		const UScriptStruct* FragmentStruct = GetFragmentStruct();
-		check(FragmentStruct);
+		const FFragmentType* ResolvedFragmentType = GetFragmentType();
+		check(ResolvedFragmentType);
 
-		this->InitializePayload(FragmentStruct);
+		this->DefaultConstructPayload(*ResolvedFragmentType);
 
-		FragmentStruct->CopyScriptStruct(this->GetPayload(), RHS.GetPayload());
+		ResolvedFragmentType->PayloadType->CopyScriptStruct(this->GetPayload(), RHS.GetPayload());
 	}
 	else
 	{
@@ -121,12 +160,13 @@ FUniversalObjectLocatorFragment& FUniversalObjectLocatorFragment::operator=(cons
 
 	if (RHS.bIsInitialized)
 	{
-		const UScriptStruct* FragmentStruct = GetFragmentStruct();
+		const FFragmentType* ResolvedFragmentType = GetFragmentType();
+		check(ResolvedFragmentType);
 
 		// Assign the FragmentType and copy the payload
 		this->FragmentType = RHS.FragmentType;
-		this->InitializePayload(FragmentStruct);
-		FragmentStruct->CopyScriptStruct(this->GetPayload(), RHS.GetPayload());
+		this->DefaultConstructPayload(*ResolvedFragmentType);
+		ResolvedFragmentType->PayloadType->CopyScriptStruct(this->GetPayload(), RHS.GetPayload());
 	}
 	else
 	{
@@ -141,14 +181,16 @@ FUniversalObjectLocatorFragment::FUniversalObjectLocatorFragment(FUniversalObjec
 	: FragmentType(RHS.FragmentType)
 	, bIsInitialized(RHS.bIsInitialized)
 	, bIsInline(RHS.bIsInline)
+	, DebugHeaderSizeLog2(RHS.DebugHeaderSizeLog2)
 {
 	using namespace UE::UniversalObjectLocator;
 
 	FMemory::Memcpy(this->Data, RHS.Data, sizeof(Data));
 
-	RHS.bIsInitialized = false;
-	RHS.bIsInline      = false;
-	RHS.FragmentType   = FFragmentTypeHandle();
+	RHS.bIsInitialized      = false;
+	RHS.bIsInline           = false;
+	RHS.DebugHeaderSizeLog2 = 0;
+	RHS.FragmentType        = FFragmentTypeHandle();
 }
 
 FUniversalObjectLocatorFragment& FUniversalObjectLocatorFragment::operator=(FUniversalObjectLocatorFragment&& RHS)
@@ -157,14 +199,17 @@ FUniversalObjectLocatorFragment& FUniversalObjectLocatorFragment::operator=(FUni
 
 	this->DestroyPayload();
 
-	this->bIsInitialized = RHS.bIsInitialized;
-	this->bIsInline      = RHS.bIsInline;
-	this->FragmentType   = RHS.FragmentType;
+	this->bIsInitialized      = RHS.bIsInitialized;
+	this->bIsInline           = RHS.bIsInline;
+	this->DebugHeaderSizeLog2 = RHS.DebugHeaderSizeLog2;
+	this->FragmentType        = RHS.FragmentType;
+
 	FMemory::Memcpy(this->Data, RHS.Data, sizeof(Data));
 
-	RHS.bIsInitialized = false;
-	RHS.bIsInline      = false;
-	RHS.FragmentType   = FFragmentTypeHandle();
+	RHS.bIsInitialized      = false;
+	RHS.bIsInline           = false;
+	RHS.DebugHeaderSizeLog2 = 0;
+	RHS.FragmentType        = FFragmentTypeHandle();
 
 	return *this;
 }
@@ -335,7 +380,7 @@ UE::UniversalObjectLocator::FParseStringResult FUniversalObjectLocatorFragment::
 		{
 			this->DestroyPayload();
 			this->FragmentType = MakeFragmentTypeHandle(SerializedFragmentType);
-			this->InitializePayload(SerializedFragmentType->PayloadType);
+			this->DefaultConstructPayload(*SerializedFragmentType);
 
 			return FParseStringResult().Success(InString.Len());
 		}
@@ -416,7 +461,7 @@ void FUniversalObjectLocatorFragment::DestroyPayload()
 		return;
 	}
 
-	void* Payload = GetPayload();
+	uint8* Payload = (uint8*)GetPayload();
 
 	const UScriptStruct* FragmentStruct = GetFragmentStruct();
 	if (ensureMsgf(FragmentStruct, TEXT("FUniversalObjectLocatorFragment has outlived its FragmentType's payload type struct! This could leak memory if the type allocated it.")))
@@ -426,6 +471,7 @@ void FUniversalObjectLocatorFragment::DestroyPayload()
 
 	if (!bIsInline)
 	{
+		Payload -= GetDebugHeaderOffset();
 		FMemory::Free(Payload);
 	}
 
@@ -435,33 +481,71 @@ void FUniversalObjectLocatorFragment::DestroyPayload()
 void* FUniversalObjectLocatorFragment::GetPayload()
 {
 	check(bIsInitialized);
-	return bIsInline ? Data : *((void**)Data);
+
+	uint8* Payload = bIsInline ? Data : *((uint8**)Data);
+	return Payload + GetDebugHeaderOffset();
 }
 
 const void* FUniversalObjectLocatorFragment::GetPayload() const
 {
-	return bIsInline ? Data : *((const void* const *)Data);
+	check(bIsInitialized);
+
+	const uint8* Payload = bIsInline ? Data : *((const uint8* const *)Data);
+	return Payload + GetDebugHeaderOffset();
 }
 
-void FUniversalObjectLocatorFragment::InitializePayload(const UScriptStruct* PayloadType)
+FUniversalObjectLocatorFragment::FAllocatedPayload FUniversalObjectLocatorFragment::AllocatePayload(size_t Size, size_t Alignment)
 {
+	using namespace UE::UniversalObjectLocator;
+
 	check(!bIsInitialized);
 
 	bIsInitialized = true;
-	if (PayloadType->GetStructureSize() <= sizeof(FUniversalObjectLocatorFragment::Data) && PayloadType->GetMinAlignment() <= alignof(FUniversalObjectLocatorFragment))
+
+#if UE_UNIVERSALOBJECTLOCATOR_DEBUG
+	DebugHeaderSizeLog2 = ComputeDebugHeaderLog2(Alignment);
+	Alignment = FMath::Max(Alignment, alignof(IFragmentPayload));
+	Size += GetDebugHeaderOffset();
+#else
+	DebugHeaderSizeLog2 = 0;
+#endif
+
+	uint8* Payload = nullptr;
+	if (Size <= sizeof(FUniversalObjectLocatorFragment::Data) && Alignment <= alignof(FUniversalObjectLocatorFragment))
 	{
 		// We can placement new this into the payload data
 		bIsInline = true;
+		Payload = Data;
 	}
 	else
 	{
 		// We have to allocate this struct on the heap
-		void* HeapAllocation = FMemory::Malloc(PayloadType->GetStructureSize(), PayloadType->GetMinAlignment());
-		*reinterpret_cast<void**>(Data) = HeapAllocation;
+		Payload = (uint8*)FMemory::Malloc(Size, Alignment);
+		*reinterpret_cast<void**>(Data) = Payload;
 		bIsInline = false;
 	}
 
-	PayloadType->InitializeStruct(GetPayload());
+	return FAllocatedPayload{
+#if UE_UNIVERSALOBJECTLOCATOR_DEBUG
+		Payload + GetDebugHeaderOffset() - sizeof(IFragmentPayload),
+#endif
+		Payload + GetDebugHeaderOffset()
+	};
+}
+
+void FUniversalObjectLocatorFragment::DefaultConstructPayload(const UE::UniversalObjectLocator::FFragmentType& InFragmentType)
+{
+	using namespace UE::UniversalObjectLocator;
+
+	const UScriptStruct* PayloadType = InFragmentType.PayloadType.Get();
+
+	FAllocatedPayload Allocation = AllocatePayload((size_t)PayloadType->GetStructureSize(), (size_t)PayloadType->GetMinAlignment());
+
+#if UE_UNIVERSALOBJECTLOCATOR_DEBUG
+	InFragmentType.StaticBindings.FragmentDebugInitializer(Allocation.DebugVFTablePtr);
+#endif
+
+	PayloadType->InitializeStruct(Allocation.Payload);
 }
 
 void FUniversalObjectLocatorFragment::Reset()
@@ -481,7 +565,7 @@ void FUniversalObjectLocatorFragment::Reset(const UObject* InObject, UObject* Co
 	if (const FFragmentType* BestFragmentType = FindBestFragmentType(InObject, Context))
 	{
 		FragmentType = MakeFragmentTypeHandle(BestFragmentType);
-		InitializePayload(BestFragmentType->PayloadType);
+		DefaultConstructPayload(*BestFragmentType);
 
 		BestFragmentType->InitializePayload(GetPayload(), FInitializeParams{ InObject, Context });
 	}
@@ -523,7 +607,7 @@ void FUniversalObjectLocatorFragment::Reset(const UObject* InObject, UObject* Co
 	if (BestFragmentType && BestFragmentType->PayloadType != nullptr)
 	{
 		FragmentType = MakeFragmentTypeHandle(BestFragmentType);
-		InitializePayload(BestFragmentType->PayloadType);
+		DefaultConstructPayload(*BestFragmentType);
 
 		BestFragmentType->InitializePayload(GetPayload(), FInitializeParams{ InObject, Context });
 	}
@@ -579,7 +663,7 @@ bool FUniversalObjectLocatorFragment::Serialize(FArchive& Ar)
 			UScriptStruct* Struct = SerializedFragmentType->GetStruct();
 			Ar.Preload(Struct);
 
-			InitializePayload(Struct);
+			DefaultConstructPayload(*SerializedFragmentType);
 			Struct->SerializeItem(Ar, GetPayload(), nullptr);
 		}
 	}
diff --git a/Engine/Source/Runtime/UniversalObjectLocator/Private/UniversalObjectLocatorModule.cpp b/Engine/Source/Runtime/UniversalObjectLocator/Private/UniversalObjectLocatorModule.cpp
--- a/Engine/Source/Runtime/UniversalObjectLocator/Private/UniversalObjectLocatorModule.cpp
+++ b/Engine/Source/Runtime/UniversalObjectLocator/Private/UniversalObjectLocatorModule.cpp
@@ -6,8 +6,7 @@
 #include "UniversalObjectLocatorFragment.h"
 #include "UniversalObjectLocatorRegistry.h"
 #include "UniversalObjectLocatorParameterTypeHandle.h"
-//#include "Modules/VisualizerDebuggingState.h"
-#include "Containers/Ticker.h"
+#include "Modules/VisualizerDebuggingState.h"
 #include "Misc/DelayedAutoRegister.h"
 
 #include "DirectPathObjectLocator.h"
@@ -20,6 +19,9 @@ class FUniversalObjectLocatorModule
 {
 public:
 	
+	/** Debug visualizer ID. String for lookup inside natvcis should be "49ceb527db044325a786a9b4470158fc" */
+	FGuid DebugVisualizerID = FGuid(0x49ceb527, 0xdb044325, 0xa786a9b4, 0x470158fc);
+
 	void StartupModule() override
 	{
 		// Register fragment types as soon as the object system is ready
@@ -38,9 +40,6 @@ public:
 				}
 			}
 		);
-
-		TickerDelegate = FTSTicker::GetCoreTicker().AddTicker(
-			FTickerDelegate::CreateRaw(this, &FUniversalObjectLocatorModule::PurgeVisualizers), 60.f);
 	}
 
 	void ShutdownModule() override
@@ -50,15 +49,16 @@ public:
 
 	FFragmentTypeHandle RegisterFragmentTypeImpl(const FFragmentType& FragmentType) override
 	{
+		using namespace UE::Core;
+
 		FRegistry& Registry = FRegistry::Get();
 
 		const int32 Index = Registry.FragmentTypes.Num();
 		Registry.FragmentTypes.Add(FragmentType);
 		checkf(Index < static_cast<int32>(std::numeric_limits<uint8>::max()), TEXT("Maximum number of UOL FragmentTypes reached"));
 
-		// @todo: enable this code once visualizer debugging state is enabled
-		// Re-assign the debugging ptr in case it changed
-		// UE::Core::FVisualizerDebuggingState::Assign("UOL", Registry.FragmentTypes.GetData());
+		// Re-assign the debugging ptr in case it changed due to reallocation
+		EVisualizerDebuggingStateResult Result = FVisualizerDebuggingState::Assign(DebugVisualizerID, Registry.FragmentTypes.GetData());
 
 		return FFragmentTypeHandle(static_cast<uint8>(Index));
 	}
@@ -86,20 +86,6 @@ public:
 		check(ParameterType.IsValid());
 		Registry.ParameterTypes[ParameterType.GetIndex()] = nullptr;
 	}
-
-	bool PurgeVisualizers(float) const
-	{
-		for (FFragmentType& FragmentType : FRegistry::Get().FragmentTypes)
-		{
-			if (FragmentType.DebuggingAssistant)
-			{
-				FragmentType.DebuggingAssistant->Purge();
-			}
-		}
-		return true;
-	}
-
-	FTSTicker::FDelegateHandle TickerDelegate;
 };
 
 } // namespace UE::UniversalObjectLocator
diff --git a/Engine/Source/Runtime/UniversalObjectLocator/Public/IUniversalObjectLocatorModule.h b/Engine/Source/Runtime/UniversalObjectLocator/Public/IUniversalObjectLocatorModule.h
--- a/Engine/Source/Runtime/UniversalObjectLocator/Public/IUniversalObjectLocatorModule.h
+++ b/Engine/Source/Runtime/UniversalObjectLocator/Public/IUniversalObjectLocatorModule.h
@@ -75,8 +75,7 @@ public:
 		FFragmentType NewFragmentType;
 		NewFragmentType.FragmentTypeID     = FragmentTypeParameters.FragmentTypeID;
 		NewFragmentType.PrimaryEditorType  = FragmentTypeParameters.PrimaryEditorType;
-		NewFragmentType.Flags			   = FragmentTypeParameters.Flags;
-		NewFragmentType.DebuggingAssistant = MakeShared<TFragmentTypeDebuggingAssistant<PayloadStructType>>();
+		NewFragmentType.Flags              = FragmentTypeParameters.Flags;
 		NewFragmentType.PayloadType        = PayloadStructType::StaticStruct();
 
 		// Static bindings
@@ -99,6 +98,12 @@ public:
 		{
 			return static_cast<PayloadStructType*>(Payload)->TryParseString(InString, Params);
 		};
+#if UE_UNIVERSALOBJECTLOCATOR_DEBUG
+		NewFragmentType.StaticBindings.FragmentDebugInitializer = [](void* Payload)
+		{
+			new (Payload) TFragmentPayload<PayloadStructType>;
+		};
+#endif
 		return TFragmentTypeHandle<PayloadStructType>(RegisterFragmentTypeImpl(NewFragmentType));
 	}
 
diff --git a/Engine/Source/Runtime/UniversalObjectLocator/Public/UniversalObjectLocatorFragment.h b/Engine/Source/Runtime/UniversalObjectLocator/Public/UniversalObjectLocatorFragment.h
--- a/Engine/Source/Runtime/UniversalObjectLocator/Public/UniversalObjectLocatorFragment.h
+++ b/Engine/Source/Runtime/UniversalObjectLocator/Public/UniversalObjectLocatorFragment.h
@@ -7,6 +7,7 @@
 #include "UniversalObjectLocatorFwd.h"
 #include "UniversalObjectLocatorFragmentType.h"
 #include "UniversalObjectLocatorFragmentTypeHandle.h"
+#include "UniversalObjectLocatorFragmentDebugging.h"
 #include "Templates/Function.h"
 #include "UniversalObjectLocatorFragment.generated.h"
 
@@ -45,9 +46,9 @@ struct alignas(8) FUniversalObjectLocatorFragment
 
 	/** Make our inline data buffer larger in-editor to support editor-only data without allocation */
 #if WITH_EDITORONLY_DATA
-	static constexpr SIZE_T Size = 64;
+	static constexpr SIZE_T SizeInMemory = 64;
 #else
-	static constexpr SIZE_T Size = 32;
+	static constexpr SIZE_T SizeInMemory = 32;
 #endif
 
 
@@ -274,10 +275,38 @@ protected:
 
 protected:
 
+	/**
+	 * Describes a UOL fragment payload after it has been allocated
+	 * Where UE_UNIVERSALOBJECTLOCATOR_DEBUG is enabled, the allocation will be laid out as such:
+	 * 0..								8..															8+sizeof(T)
+	 * [TFragmentPayload<T>				| T Payload													]
+	 *
+	 * When UE_UNIVERSALOBJECTLOCATOR_DEBUG is disabled, the allocation is simply:
+	 * 0..															sizeof(T)
+	 * | T Payload													]
+	 */
+	struct FAllocatedPayload
+	{
+#if UE_UNIVERSALOBJECTLOCATOR_DEBUG
+		/** Pointer to artificially inserted vtable ptr that identifies the fragment data polymorphically */
+		void* DebugVFTablePtr;
+#endif
+		/** Pointer to the actual typed fragment data */
+		void* Payload;
+	};
+
+	/**
+	 * Allocate (but do not initialize) the fragment payload using the specified size and alignment,
+	 *     accounting for enough space to fit a debug vftable ptr whern UE_UNIVERSALOBJECTLOCATOR_DEBUG is enabled.
+	 * 
+	 * @return Structure identifying the payload
+	 */
+	UNIVERSALOBJECTLOCATOR_API FAllocatedPayload AllocatePayload(size_t Size, size_t Alignment);
+
 	/**
 	 * Default-initialize the fragment payload using the specified type
 	 */
-	void InitializePayload(const UScriptStruct* StructType);
+	void DefaultConstructPayload(const UE::UniversalObjectLocator::FFragmentType& InFragmentType);
 
 	/**
 	 * Destroy the payload (if valid) by calling its destructor and freeing the memory (if necessary)
@@ -286,18 +315,32 @@ protected:
 
 private:
 
+	uint32 GetDebugHeaderOffset() const
+	{
+		// Special case for DebugHeaderSizeLog2==0 which signifies no offset rather than 2^0 = 1 byte.
+		// We do this using a branchless bitmask that always unsets the first bit (which can never be set, because we always return a power of 2 > 1)
+		return (1ul << DebugHeaderSizeLog2) & (~1ul);
+	}
+
+#if UE_UNIVERSALOBJECTLOCATOR_DEBUG
 	/*~ Utility symbol name to guarantee that FFragmentType can be resolved within the context of a FUniversalObjectLocatorFragment within natvis expressions */
-	struct FDebuggableFragmentType : UE::UniversalObjectLocator::FFragmentType
+	struct FDebuggableFragmentType
 	{
+		using Type = UE::UniversalObjectLocator::FFragmentType;
 	};
+	struct FDebuggableFragment
+	{
+		using Type = UE::UniversalObjectLocator::IFragmentPayload;
+	};
+#endif
 
 	/*
 	 * Payload data - implicitly aligned to a 8 byte boundary since it's the first member.
 	 * Given payload type T, this is either a type-erased T() value (where bIsInline==1),
 	 *    or a T* to a heap allocated T (where bIsInline==0)
-	 * Size is specifically defined by the desired overall size of FUniversalObjectLocatorFragment::Size, minus space for other members
+	 * SizeInMemory is specifically defined by the desired overall size of FUniversalObjectLocatorFragment::SizeInMemory, minus space for other members
 	 */
-	uint8 Data[Size-2];
+	uint8 Data[SizeInMemory-2];
 
 	/** 1 Byte - the fragment type portion of the universal fragment */
 	UE::UniversalObjectLocator::FFragmentTypeHandle FragmentType;
@@ -307,6 +350,8 @@ private:
 	uint8 bIsInitialized : 1;
 	/** True if Data is an inline allocation of FragmentType::PayloadType, false means Data is a (void*) to the heap allocated data. */
 	uint8 bIsInline : 1;
+	/** Offset from the allocated memory to the fragment payload stored as a power of 2. Only non-zero when UE_UNIVERSALOBJECTLOCATOR_DEBUG is enabled. */
+	uint8 DebugHeaderSizeLog2 : 6;
 };
 
 template<>
@@ -327,20 +372,22 @@ struct TStructOpsTypeTraits<FUniversalObjectLocatorFragment> : public TStructOps
 template<typename T, typename ...ArgTypes>
 FUniversalObjectLocatorFragment::FUniversalObjectLocatorFragment(UE::UniversalObjectLocator::TFragmentTypeHandle<T> InHandle, ArgTypes&& ...InArgs)
 	: FragmentType(InHandle)
-	, bIsInitialized(1)
+	, bIsInitialized(0)
+	, DebugHeaderSizeLog2(0)
 {
+	using namespace UE::UniversalObjectLocator;
+
 	checkf(InHandle, TEXT("Attempting to construct a new fragment from an invalid fragment type handle - was it registered?"));
 
-	bIsInline = sizeof(T) <= sizeof(Data) && alignof(T) <= alignof(FUniversalObjectLocatorFragment);
-	if (!bIsInline)
-	{
-		// We have to allocate this struct on the heap
-		void* HeapAllocation = FMemory::Malloc(sizeof(T), alignof(T));
-		*reinterpret_cast<void**>(Data) = HeapAllocation;
-	}
+	FAllocatedPayload Allocation = AllocatePayload(sizeof(T), alignof(T));
+
+#if UE_UNIVERSALOBJECTLOCATOR_DEBUG
+	// Initialize the fragment vftable if necessary. We can do this without needing the fragment type since we know the type
+	new (Allocation.DebugVFTablePtr) TFragmentPayload<T>;
+#endif
 
 	// Placement new the payload
-	new (GetPayload()) T{ Forward<ArgTypes>(InArgs)... };
+	new (Allocation.Payload) T{ Forward<ArgTypes>(InArgs)... };
 }
 
 template<typename T>
diff --git a/Engine/Source/Runtime/UniversalObjectLocator/Public/UniversalObjectLocatorFragmentDebugging.h b/Engine/Source/Runtime/UniversalObjectLocator/Public/UniversalObjectLocatorFragmentDebugging.h
new file mode 100644
--- /dev/null
+++ b/Engine/Source/Runtime/UniversalObjectLocator/Public/UniversalObjectLocatorFragmentDebugging.h
@@ -0,0 +1,41 @@
+// Copyright Epic Games, Inc. All Rights Reserved.
+
+#pragma once
+
+#include "CoreTypes.h"
+
+#ifndef UE_UNIVERSALOBJECTLOCATOR_DEBUG
+	#define UE_UNIVERSALOBJECTLOCATOR_DEBUG !UE_BUILD_SHIPPING
+#endif
+
+#if UE_UNIVERSALOBJECTLOCATOR_DEBUG
+
+namespace UE::UniversalObjectLocator
+{
+
+	/**
+	 * Type whose sole purpose is to add a vtable pointer in front of a fragment
+	 *      to assist in debugging. When allocating the fragment, a TFragmentPayload<T>
+	 *      is allocated in the preceeding 8 bytes which can be used by a natvis expression
+	 *      to show the proceeding bytes as a (T*)
+	 */
+	struct alignas(8) IFragmentPayload
+	{
+		virtual ~IFragmentPayload()
+		{
+		}
+	};
+
+	/**
+	 * Templated version of IFragmentPayload that is added to the start of a fragment.
+	 * Utilizes a zero-sized array to cast the proceeding bytes to a T*.
+	 */
+	template<typename T>
+	struct TFragmentPayload : IFragmentPayload
+	{
+		UE_NO_UNIQUE_ADDRESS uint8 Ptr[];
+	};
+
+} // namespace UE::UniversalObjectLocator
+
+#endif // UE_UNIVERSALOBJECTLOCATOR_DEBUG
diff --git a/Engine/Source/Runtime/UniversalObjectLocator/Public/UniversalObjectLocatorFragmentType.h b/Engine/Source/Runtime/UniversalObjectLocator/Public/UniversalObjectLocatorFragmentType.h
--- a/Engine/Source/Runtime/UniversalObjectLocator/Public/UniversalObjectLocatorFragmentType.h
+++ b/Engine/Source/Runtime/UniversalObjectLocator/Public/UniversalObjectLocatorFragmentType.h
@@ -14,6 +14,7 @@
 
 #include "UniversalObjectLocatorFwd.h"
 #include "UniversalObjectLocatorResolveParams.h"
+#include "UniversalObjectLocatorFragmentDebugging.h"
 
 namespace UE::UniversalObjectLocator
 {
@@ -40,6 +41,7 @@ enum class EFragmentTypeFlags : uint8
 
 ENUM_CLASS_FLAGS(EFragmentTypeFlags)
 
+
 /**
  * A Universal Object Locator fragment type defines a specific mechanism for resolving an object that is referenced
  *   inside a Universal Object Locator (which can be considered a chain of fragments). Each fragment
@@ -80,6 +82,10 @@ struct FFragmentType
 
 	using FPriorityCallback   = uint32 (*)(const UObject* /* Object */, const UObject* /* Context */);
 
+#if UE_UNIVERSALOBJECTLOCATOR_DEBUG
+	using FFragmentDebugInitializer = void (*)(void*);
+#endif
+
 	/** Bindings structure for function pointers that directly invoke behavior on fragment type payloads */
 	struct
 	{
@@ -98,14 +104,15 @@ struct FFragmentType
 	{
 		/** static function binding for computing the priority of this fragment type */
 		FPriorityCallback   Priority;
+#if UE_UNIVERSALOBJECTLOCATOR_DEBUG
+		/** Function callback for initializing a vftable for the fragment for debugging purposes */
+		FFragmentDebugInitializer FragmentDebugInitializer;
+#endif
 	} StaticBindings;
 
 	/** Struct pointer to the payload type */
 	TObjectPtr<const UScriptStruct> PayloadType;
 
-	/** Debugging assistant structure used for natvis visualizations */
-	TSharedPtr<IFragmentTypeDebuggingAssistant> DebuggingAssistant;
-
 	/** Descriptive display text for this type of fragment type */
 	FText DisplayText;
 
@@ -129,46 +136,4 @@ struct TFragmentType : FFragmentType
 {
 };
 
-struct IPayloadDebugging
-{
-	virtual ~IPayloadDebugging()
-	{
-	}
-};
-
-template<typename T>
-struct TLocatorFragment : IPayloadDebugging
-{
-	explicit TLocatorFragment(const T* InPayload)
-		: Payload(InPayload)
-	{
-	}
-
-	const T* Payload;
-};
-
-struct IFragmentTypeDebuggingAssistant : TSharedFromThis<IFragmentTypeDebuggingAssistant>
-{
-	virtual ~IFragmentTypeDebuggingAssistant()
-	{
-	}
-	virtual IPayloadDebugging* MakePayloadVisualizer(const void*) = 0;
-	virtual void Purge() = 0;
-};
-
-template<typename T>
-struct TFragmentTypeDebuggingAssistant : IFragmentTypeDebuggingAssistant
-{
-	IPayloadDebugging* MakePayloadVisualizer(const void* InPayload) override
-	{
-		return &Visualizers.Emplace_GetRef(static_cast<const T*>(InPayload));
-	}
-	void Purge() override
-	{
-		Visualizers.Empty();
-	}
-
-	TPagedArray<TLocatorFragment<T>> Visualizers;
-};
-
 } // namespace UE::UniversalObjectLocator
\ No newline at end of file
