Changed UOL check for a non-fatal ensure and MessageLog

diff --git a/Engine/Source/Editor/UniversalObjectLocatorEditor/Private/UniversalObjectLocatorEditorModule.cpp b/Engine/Source/Editor/UniversalObjectLocatorEditor/Private/UniversalObjectLocatorEditorModule.cpp
--- a/Engine/Source/Editor/UniversalObjectLocatorEditor/Private/UniversalObjectLocatorEditorModule.cpp
+++ b/Engine/Source/Editor/UniversalObjectLocatorEditor/Private/UniversalObjectLocatorEditorModule.cpp
@@ -11,6 +11,8 @@
 
 #include "ISequencerModule.h"
 
+#include "MessageLogModule.h"
+
 namespace UE::UniversalObjectLocator
 {
 
@@ -25,6 +27,9 @@ void FUniversalObjectLocatorEditorModule::StartupModule()
 		RegisterLocatorEditor("AnimInstance", MakeShared<FAnimInstanceLocatorEditor>());
 		RegisterLocatorEditor("Asset", MakeShared<FAssetLocatorEditor>());
 	}
+
+	FMessageLogModule& MessageLogModule = FModuleManager::LoadModuleChecked<FMessageLogModule>("MessageLog");
+	MessageLogModule.RegisterLogListing("UOL", NSLOCTEXT("UOL", "UOLMessageLabel", "Universal Object Locator"));
 }
 
 void FUniversalObjectLocatorEditorModule::ShutdownModule()
diff --git a/Engine/Source/Editor/UniversalObjectLocatorEditor/UniversalObjectLocatorEditor.Build.cs b/Engine/Source/Editor/UniversalObjectLocatorEditor/UniversalObjectLocatorEditor.Build.cs
--- a/Engine/Source/Editor/UniversalObjectLocatorEditor/UniversalObjectLocatorEditor.Build.cs
+++ b/Engine/Source/Editor/UniversalObjectLocatorEditor/UniversalObjectLocatorEditor.Build.cs
@@ -22,6 +22,7 @@ public class UniversalObjectLocatorEditor : ModuleRules
 				"CoreUObject",
 				"EditorWidgets",
 				"Engine",
+				"MessageLog",
 				"PropertyEditor",
 				"SceneOutliner",
 				"Slate",
diff --git a/Engine/Source/Runtime/UniversalObjectLocator/Private/UniversalObjectLocatorFragment.cpp b/Engine/Source/Runtime/UniversalObjectLocator/Private/UniversalObjectLocatorFragment.cpp
--- a/Engine/Source/Runtime/UniversalObjectLocator/Private/UniversalObjectLocatorFragment.cpp
+++ b/Engine/Source/Runtime/UniversalObjectLocator/Private/UniversalObjectLocatorFragment.cpp
@@ -9,8 +9,11 @@
 #include "UniversalObjectLocatorFragmentDebugging.h"
 #include "UniversalObjectLocatorRegistry.h"
 #include "UObject/SoftObjectPath.h"
+#include "UObject/UObjectThreadContext.h"
 #include "Containers/SparseArray.h"
+#include "Logging/MessageLog.h"
 #include "Misc/AsciiSet.h"
+#include "Misc/UObjectToken.h"
 #include "Templates/AlignmentTemplates.h"
 
 #define LOCTEXT_NAMESPACE "UOL"
@@ -649,8 +652,16 @@ bool FUniversalObjectLocatorFragment::Serialize(FArchive& Ar)
 				Reset();
 
 				// Big error - what do we do?
-				UE_LOG(LogUniversalObjectLocator, Error, TEXT("WARNING: POTENTIAL DATA LOSS! Universal Object Reference FragmentType %s is not recognized! This reference will be lost if re-saved."), *FragmentTypeID.ToString());
-				checkf(false, TEXT("WARNING: POTENTIAL DATA LOSS! Universal Object Reference FragmentType %s! This reference will be lost if re-saved."), *FragmentTypeID.ToString());
+				FMessageLog Log("UOL");
+				TSharedRef<FTokenizedMessage> Message = Log.Error(FText::Format(LOCTEXT("DataLossWarning", "WARNING: POTENTIAL DATA LOSS! Universal Object Reference FragmentType {0}! This reference will be lost if re-saved."), FText::FromString(FragmentTypeID.ToString())));
+
+				FUObjectSerializeContext* SerializeContext = FUObjectThreadContext::Get().GetSerializeContext();
+				if (SerializeContext && SerializeContext->SerializedObject)
+				{
+					Message->AddToken(FUObjectToken::Create(SerializeContext->SerializedObject));
+				}
+
+				Log.Open(EMessageSeverity::Error);
 
 				// Deserialize an empty payload so we don't corrupt the serialization data
 				FUniversalObjectLocatorEmptyPayload Empty;
