UOL: Guard against null terminators appearing in string represnetations

There is currently an issue where soft object paths can have two null terminators in their string, which results in a null terminator being added to the string representation, which then fails the character validation of UOL fragments.

For now we will strip these null terminators from the string, and I have improved the ensure message to include a description of the failing characters.

diff --git a/Engine/Source/Runtime/UniversalObjectLocator/Private/UniversalObjectLocatorFragmentType.cpp b/Engine/Source/Runtime/UniversalObjectLocator/Private/UniversalObjectLocatorFragmentType.cpp
--- a/Engine/Source/Runtime/UniversalObjectLocator/Private/UniversalObjectLocatorFragmentType.cpp
+++ b/Engine/Source/Runtime/UniversalObjectLocator/Private/UniversalObjectLocatorFragmentType.cpp
@@ -19,13 +19,44 @@ FInitializeResult FFragmentType::InitializePayload(void* Payload, const FInitial
 }
 void FFragmentType::ToString(const void* Payload, FStringBuilderBase& OutStringBuilder) const
 {
+	// Some strings are currently emitting null terminators due to a bug with UTF8
+	//   Strip these off to ensure we don't fail the check below
+	while (OutStringBuilder.Len() != 0 && OutStringBuilder.LastChar() == '\0')
+	{
+		OutStringBuilder.RemoveSuffix(1);
+	}
+
 	const int32 StartPos = OutStringBuilder.Len();
 
 	(*InstanceBindings.ToString)(Payload, OutStringBuilder);
 
 #if DO_CHECK
+
+	static constexpr FAsciiSet InvalidChars = ~FUniversalObjectLocatorFragment::ValidFragmentPayloadCharacters;
+
+	auto MakeErrorString = [](FStringView StringView){
+		
+		FString String;
+		for (TCHAR Char : StringView)
+		{
+			if (InvalidChars.Test(Char))
+			{
+				String.AppendChar('^');
+			}
+			else
+			{
+				String.AppendChar(' ');
+			}
+		}
+		return String;
+	};
+
 	FStringView StringRepresentation(OutStringBuilder.GetData() + StartPos, OutStringBuilder.Len());
-	checkf(!FAsciiSet::HasAny(StringRepresentation, ~FUniversalObjectLocatorFragment::ValidFragmentPayloadCharacters), TEXT("F%s::ToString resulted in an invalid character usage"), *PayloadType->GetName());
+	checkf(!FAsciiSet::HasAny(StringRepresentation, InvalidChars),
+		TEXT("F%s::ToString resulted in an invalid character usage:\n\t%s\n\t%s"),
+		*PayloadType->GetName(),
+		*FString(StringRepresentation.Len(), StringRepresentation.GetData()),
+		*MakeErrorString(StringRepresentation));
 #endif
 }
 FParseStringResult FFragmentType::TryParseString(void* Payload, FStringView InString, const FParseStringParams& Params) const
