PVS 7.38 fix ups

#rnx
#DisabledValidations Preflight

diff --git a/Engine/Plugins/Animation/ControlRig/Source/ControlRig/Private/Sequencer/MovieSceneControlRigSystem.cpp b/Engine/Plugins/Animation/ControlRig/Source/ControlRig/Private/Sequencer/MovieSceneControlRigSystem.cpp
--- a/Engine/Plugins/Animation/ControlRig/Source/ControlRig/Private/Sequencer/MovieSceneControlRigSystem.cpp
+++ b/Engine/Plugins/Animation/ControlRig/Source/ControlRig/Private/Sequencer/MovieSceneControlRigSystem.cpp
@@ -1927,13 +1927,10 @@ void UMovieSceneControlRigParameterEvaluatorSystem::InitializeBaseRigComponent(
 		TArray<FName> NewSelectedControls = ControlRig->CurrentControlSelection();
 		if (SelectedControls != NewSelectedControls)
 		{
-			if (ControlRig)
+			ControlRig->ClearControlSelection();
+			for (const FName& Name : SelectedControls)
 			{
-				ControlRig->ClearControlSelection();
-				for (const FName& Name : SelectedControls)
-				{
-					ControlRig->SelectControl(Name, true);
-				}
+				ControlRig->SelectControl(Name, true);
 			}
 		}
 	}
diff --git a/Engine/Plugins/Animation/PoseSearch/Source/Runtime/Private/PoseSearchDerivedData.cpp b/Engine/Plugins/Animation/PoseSearch/Source/Runtime/Private/PoseSearchDerivedData.cpp
--- a/Engine/Plugins/Animation/PoseSearch/Source/Runtime/Private/PoseSearchDerivedData.cpp
+++ b/Engine/Plugins/Animation/PoseSearch/Source/Runtime/Private/PoseSearchDerivedData.cpp
@@ -2505,7 +2505,7 @@ void FAsyncPoseSearchDatabasesManagement::Tick(float DeltaTime)
 												const UAnimationAsset* AnimationAsset = DatabaseAsset->GetAnimationAssetForRole(Role);
 												const FTransform RootTransformOrigin = DatabaseAsset->GetRootTransformOriginForRole(Role);
 
-												static int32 NUM_TESTING_SAMPLES_PER_DATABASE_ASSET = 100; // make sure it's greater than 1
+												constexpr int32 NUM_TESTING_SAMPLES_PER_DATABASE_ASSET = 100; // make sure it's greater than 1
 
 												for (int32 SampleIndex = 0; SampleIndex < NUM_TESTING_SAMPLES_PER_DATABASE_ASSET; SampleIndex++)
 												{
diff --git a/Engine/Plugins/ChaosClothAsset/Source/ChaosClothAsset/Private/ChaosClothAsset/ClothCollection.cpp b/Engine/Plugins/ChaosClothAsset/Source/ChaosClothAsset/Private/ChaosClothAsset/ClothCollection.cpp
--- a/Engine/Plugins/ChaosClothAsset/Source/ChaosClothAsset/Private/ChaosClothAsset/ClothCollection.cpp
+++ b/Engine/Plugins/ChaosClothAsset/Source/ChaosClothAsset/Private/ChaosClothAsset/ClothCollection.cpp
@@ -899,7 +899,7 @@ namespace UE::Chaos::ClothAsset
 
 	/*static*/ int32 FConstClothCollection::GetElementsOffset(const TManagedArray<int32>* StartArray, int32 BaseElementIndex, int32 ElementIndex)
 	{
-		while ((*StartArray)[BaseElementIndex] == INDEX_NONE && BaseElementIndex < ElementIndex)
+		while (BaseElementIndex < ElementIndex && (*StartArray)[BaseElementIndex] == INDEX_NONE)
 		{
 			++BaseElementIndex;
 		}
diff --git a/Engine/Plugins/Experimental/Fracture/Source/FractureEngine/Private/FractureEngineMaterials.cpp b/Engine/Plugins/Experimental/Fracture/Source/FractureEngine/Private/FractureEngineMaterials.cpp
--- a/Engine/Plugins/Experimental/Fracture/Source/FractureEngine/Private/FractureEngineMaterials.cpp
+++ b/Engine/Plugins/Experimental/Fracture/Source/FractureEngine/Private/FractureEngineMaterials.cpp
@@ -12,7 +12,7 @@ namespace UE::FractureEngineMaterialHelpers
 		bool bInternalFaces = TargetFaces == FFractureEngineMaterials::ETargetFaces::InternalFaces;
 		for (int32 FIdx = Start; FIdx < End; ++FIdx)
 		{
-			if (bAllFaces || bInternalFaces == Internal[FIdx])
+			if (bAllFaces || bInternalFaces == Internal[FIdx]) //-V562
 			{
 				Material[FIdx] = MaterialID;
 			}
diff --git a/Engine/Plugins/Experimental/GeometryCollectionPlugin/Source/GeometryCollectionNodes/Private/Dataflow/GeometryCollectionSelectionNodes.cpp b/Engine/Plugins/Experimental/GeometryCollectionPlugin/Source/GeometryCollectionNodes/Private/Dataflow/GeometryCollectionSelectionNodes.cpp
--- a/Engine/Plugins/Experimental/GeometryCollectionPlugin/Source/GeometryCollectionNodes/Private/Dataflow/GeometryCollectionSelectionNodes.cpp
+++ b/Engine/Plugins/Experimental/GeometryCollectionPlugin/Source/GeometryCollectionNodes/Private/Dataflow/GeometryCollectionSelectionNodes.cpp
@@ -1533,8 +1533,8 @@ static void CreateSelectionFromAttr(const FManagedArrayCollection& InCollection,
 
 		for (int32 Idx = 0; Idx < NumElements; ++Idx)
 		{
-			if ((InOperation == ESelectionByAttrOperation::Equal && (*Array)[Idx] == BoolValue) ||
-				(InOperation == ESelectionByAttrOperation::NotEqual && !((*Array)[Idx] == BoolValue)))
+			if ((InOperation == ESelectionByAttrOperation::Equal && (*Array)[Idx] == BoolValue) || //-V562
+				(InOperation == ESelectionByAttrOperation::NotEqual && !((*Array)[Idx] == BoolValue))) //-V562
 			{
 				OutSelection.SetSelected(Idx);
 			}
diff --git a/Engine/Plugins/Experimental/MeshModelingToolsetExp/Source/MeshModelingToolsExp/Private/RevolveSplineTool.cpp b/Engine/Plugins/Experimental/MeshModelingToolsetExp/Source/MeshModelingToolsExp/Private/RevolveSplineTool.cpp
--- a/Engine/Plugins/Experimental/MeshModelingToolsetExp/Source/MeshModelingToolsExp/Private/RevolveSplineTool.cpp
+++ b/Engine/Plugins/Experimental/MeshModelingToolsetExp/Source/MeshModelingToolsExp/Private/RevolveSplineTool.cpp
@@ -152,7 +152,7 @@ void URevolveSplineTool::OnSplineUpdate()
 		double SegmentLengthToUse = Length / NumSegments;
 		for (int32 i = 0; i <= NumSegments; ++i)
 		{
-			double DistanceToSampleAt = Length * ((double)i / NumSegments);
+			double DistanceToSampleAt = Length * ((double)i / NumSegments); //-V609
 			ProfileCurve.Add(Spline->GetLocationAtDistanceAlongSpline(DistanceToSampleAt, 
 				ESplineCoordinateSpace::World));
 		}
diff --git a/Engine/Plugins/Experimental/PlanarCutPlugin/Source/PlanarCut/Private/FractureAutoUV.cpp b/Engine/Plugins/Experimental/PlanarCutPlugin/Source/PlanarCut/Private/FractureAutoUV.cpp
--- a/Engine/Plugins/Experimental/PlanarCutPlugin/Source/PlanarCut/Private/FractureAutoUV.cpp
+++ b/Engine/Plugins/Experimental/PlanarCutPlugin/Source/PlanarCut/Private/FractureAutoUV.cpp
@@ -924,7 +924,7 @@ bool TextureSpecifiedFaces(
 				const int32 CollectionFaceIndex = MeshData.GetCollectionFaceFromMeshFace(TID);
 				if (ensure(CollectionFaceIndex != INDEX_NONE && ToTextureTriangles.IsValidIndex(CollectionFaceIndex)))
 				{
-					if (ToTextureTriangles[CollectionFaceIndex] == false)
+					if (ToTextureTriangles[CollectionFaceIndex] == false) //-V562
 					{
 						UV->UnsetTriangle(TID);
 					}
diff --git a/Engine/Plugins/Runtime/GeometryProcessing/Source/DynamicMesh/Private/Parameterization/DynamicMeshUVEditor.cpp b/Engine/Plugins/Runtime/GeometryProcessing/Source/DynamicMesh/Private/Parameterization/DynamicMeshUVEditor.cpp
--- a/Engine/Plugins/Runtime/GeometryProcessing/Source/DynamicMesh/Private/Parameterization/DynamicMeshUVEditor.cpp
+++ b/Engine/Plugins/Runtime/GeometryProcessing/Source/DynamicMesh/Private/Parameterization/DynamicMeshUVEditor.cpp
@@ -1596,7 +1596,7 @@ bool FDynamicMeshUVEditor::AutoOrientUVArea(const TArray<int32>& Triangles)
 	bool bFoundSmallerBox = false;
 	for (int32 j = 0; j <= NV; ++j)
 	{
-		FVector2f A = HullPoints[j%NV], B = HullPoints[(j+1)%NV], C = HullPoints[(j+NV/2)%NV];
+		FVector2f A = HullPoints[j%NV], B = HullPoints[(j+1)%NV], C = HullPoints[(j+NV/2)%NV]; //-V609
 		FVector2f Axis0 = B - A; 
 		float Dimension0 = Normalize(Axis0);
 		FVector2f Axis1 = PerpCW(Axis0);
diff --git a/Engine/Source/Developer/PakFileUtilities/Private/PakFileUtilities.cpp b/Engine/Source/Developer/PakFileUtilities/Private/PakFileUtilities.cpp
--- a/Engine/Source/Developer/PakFileUtilities/Private/PakFileUtilities.cpp
+++ b/Engine/Source/Developer/PakFileUtilities/Private/PakFileUtilities.cpp
@@ -1059,6 +1059,7 @@ bool PrepareCopyFileToPak(const FString& InMountPoint, const FPakInputPair& InFi
 		uint8* UncompressedBuffer = UncompressedFile.CompressedBuffer.Get();
 		if (InFile.bNeedEncryption)
 		{
+			checkSlow(FileSize > 0 || PaddedEncryptedFileSize == 0);
 			for(int64 FillIndex = FileSize; FillIndex < PaddedEncryptedFileSize; ++FillIndex)
 			{
 				// Fill the trailing buffer with bytes from file. Note that this is now from a fixed location
diff --git a/Engine/Source/Developer/ShaderPreprocessor/Private/stb_preprocess/preprocessor.c b/Engine/Source/Developer/ShaderPreprocessor/Private/stb_preprocess/preprocessor.c
--- a/Engine/Source/Developer/ShaderPreprocessor/Private/stb_preprocess/preprocessor.c
+++ b/Engine/Source/Developer/ShaderPreprocessor/Private/stb_preprocess/preprocessor.c
@@ -4732,7 +4732,7 @@ static void transition(int a, int b, int c)
 	else
 	{
 		if (b == DEFAULT)
-			memset(pp_transition_table[a], c, PP_CHAR_CLASS_count); //-V557
+			memset(pp_transition_table[a], c, PP_CHAR_CLASS_count);
 		else if (b == ANY_LETTER)
 		{
 			transition(a, PP_CHAR_CLASS_idtype, c);
@@ -4761,7 +4761,7 @@ static void initscan_transition(int a, int b, int c)
 	else
 	{
 		if (b == DEFAULT)
-			memset(pp_transition_initscan[a], c, PP_CHAR_CLASS_count); //-V557
+			memset(pp_transition_initscan[a], c, PP_CHAR_CLASS_count);
 		else if (b == ANY_LETTER)
 		{
 			initscan_transition(a, PP_CHAR_CLASS_idtype, c);
@@ -4834,6 +4834,9 @@ void init_preprocessor_scanner(void)
 	for (i = 0; i < PP_STATE_active_count; ++i)
 		transition(i, DEFAULT, PP_STATE_eof);
 
+	// Disable V557 due to false positive and the use of STATE_DEFAULT
+	//-V:transition:557
+
 	// every character class has a default behavior:
 	transition(STATE_DEFAULT, PP_CHAR_CLASS_default, PP_STATE_ready);
 	transition(STATE_DEFAULT, PP_CHAR_CLASS_whitespace, PP_STATE_ready);
diff --git a/Engine/Source/Runtime/Core/Private/Misc/AES.cpp b/Engine/Source/Runtime/Core/Private/Misc/AES.cpp
--- a/Engine/Source/Runtime/Core/Private/Misc/AES.cpp
+++ b/Engine/Source/Runtime/Core/Private/Misc/AES.cpp
@@ -201,7 +201,7 @@ static struct FAesRuntimeHWSupportDetection
 {
 	FAesRuntimeHWSupportDetection()
 	{
-		PickImplementation();
+		PickImplementation(); //-V607
 	}
 } AesRuntimeHWSupportDetection;
 
diff --git a/Engine/Source/Runtime/Core/Private/ProfilingDebugging/PlatformEvents.cpp b/Engine/Source/Runtime/Core/Private/ProfilingDebugging/PlatformEvents.cpp
--- a/Engine/Source/Runtime/Core/Private/ProfilingDebugging/PlatformEvents.cpp
+++ b/Engine/Source/Runtime/Core/Private/ProfilingDebugging/PlatformEvents.cpp
@@ -98,7 +98,7 @@ UE_TRACE_EVENT_END()
 
 void FPlatformEventsTrace::OutputSettings(uint32 SamplingIntervalUsec)
 {
-	bool bShouldTrace = UE_TRACE_CHANNELEXPR_IS_ENABLED(StackSamplingChannel) || UE_TRACE_CHANNELEXPR_IS_ENABLED(ContextSwitchChannel);
+	bool bShouldTrace = UE_TRACE_CHANNELEXPR_IS_ENABLED(StackSamplingChannel) || UE_TRACE_CHANNELEXPR_IS_ENABLED(ContextSwitchChannel); //-V501
 	UE_TRACE_LOG(PlatformEvent, Settings, bShouldTrace)
 		<< Settings.SamplingInterval(SamplingIntervalUsec);
 }
@@ -128,7 +128,7 @@ void FPlatformEventsTrace::OutputStackSample(uint64 Time, uint32 ThreadId, const
 
 void FPlatformEventsTrace::OutputThreadName(uint32 ThreadId, uint32 ProcessId, const TCHAR* Name, uint32 NameLen)
 {
-	bool bShouldTrace = UE_TRACE_CHANNELEXPR_IS_ENABLED(StackSamplingChannel) || UE_TRACE_CHANNELEXPR_IS_ENABLED(ContextSwitchChannel);
+	bool bShouldTrace = UE_TRACE_CHANNELEXPR_IS_ENABLED(StackSamplingChannel) || UE_TRACE_CHANNELEXPR_IS_ENABLED(ContextSwitchChannel); //-V501
 	UE_TRACE_LOG(PlatformEvent, ThreadName, bShouldTrace)
 		<< ThreadName.ThreadId(ThreadId)
 		<< ThreadName.ProcessId(ProcessId)
diff --git a/Engine/Source/Runtime/Core/Public/Templates/EqualTo.h b/Engine/Source/Runtime/Core/Public/Templates/EqualTo.h
--- a/Engine/Source/Runtime/Core/Public/Templates/EqualTo.h
+++ b/Engine/Source/Runtime/Core/Public/Templates/EqualTo.h
@@ -25,7 +25,7 @@ struct TEqualTo
 		else
 		{
 			static_assert(sizeof(T) == 0, "Trying to use TEqualTo<T> where T doesn't have an appropriate operator== overload. Please add bool operator==(T, T), do not specialize TEqualTo.");
-		}
+		} //-V591
 	}
 };
 
@@ -42,6 +42,6 @@ struct TEqualTo<void>
 		else
 		{
 			static_assert(sizeof(T) == 0, "Trying to use TEqualTo<void> with types without an appropriate operator== overload. Please add bool operator==(T, U), do not specialize TEqualTo.");
-		}
+		} //-V591
 	}
 };
diff --git a/Engine/Source/Runtime/Core/Public/Templates/Greater.h b/Engine/Source/Runtime/Core/Public/Templates/Greater.h
--- a/Engine/Source/Runtime/Core/Public/Templates/Greater.h
+++ b/Engine/Source/Runtime/Core/Public/Templates/Greater.h
@@ -26,7 +26,7 @@ struct TGreater
 		else
 		{
 			static_assert(sizeof(T) == 0, "Trying to use TGreater<T> where T doesn't have an appropriate operator< overload. Please add bool operator<(T, T), do not specialize TGreater.");
-		}
+		} //-V591
 	}
 };
 
@@ -43,6 +43,6 @@ struct TGreater<void>
 		else
 		{
 			static_assert(sizeof(T) == 0, "Trying to use TGreater<void> with types without an appropriate operator< overload. Please add bool operator<(T, U), do not specialize TGreater.");
-		}
+		} //-V591
 	}
 };
diff --git a/Engine/Source/Runtime/Core/Public/Templates/Less.h b/Engine/Source/Runtime/Core/Public/Templates/Less.h
--- a/Engine/Source/Runtime/Core/Public/Templates/Less.h
+++ b/Engine/Source/Runtime/Core/Public/Templates/Less.h
@@ -26,7 +26,7 @@ struct TLess
 		else
 		{
 			static_assert(sizeof(T) == 0, "Trying to use TLess<T> where T doesn't have an appropriate operator< overload. Please add bool operator<(T, T), do not specialize TLess.");
-		}
+		} //-V591
 	}
 };
 
@@ -43,6 +43,6 @@ struct TLess<void>
 		else
 		{
 			static_assert(sizeof(T) == 0, "Trying to use TLess<void> with types without an appropriate operator< overload. Please add bool operator<(T, U), do not specialize TLess.");
-		}
+		} //-V591
 	}
 };
diff --git a/Engine/Source/Runtime/Engine/Private/Particles/ParticleModules_Camera.cpp b/Engine/Source/Runtime/Engine/Private/Particles/ParticleModules_Camera.cpp
--- a/Engine/Source/Runtime/Engine/Private/Particles/ParticleModules_Camera.cpp
+++ b/Engine/Source/Runtime/Engine/Private/Particles/ParticleModules_Camera.cpp
@@ -76,10 +76,13 @@ void UParticleModuleCameraOffset::Spawn(const FSpawnContext& Context)
 
 	float ScaleFactor = 1.0f;
 
-	UParticleLODLevel* LODLevel = Owner ? Owner->SpriteTemplate->GetCurrentLODLevel(Owner) : NULL;
-	if (LODLevel && LODLevel->RequiredModule->bUseLocalSpace == false)
+	if (Owner)
 	{
-		ScaleFactor = Owner->Component.GetAsyncComponentToWorld().GetMaximumAxisScale();
+		UParticleLODLevel* LODLevel = Owner->SpriteTemplate->GetCurrentLODLevel(Owner);
+		if (LODLevel && LODLevel->RequiredModule->bUseLocalSpace == false)
+		{
+			ScaleFactor = Owner->Component.GetAsyncComponentToWorld().GetMaximumAxisScale();
+		}
 	}
 	SPAWN_INIT;
 	{
diff --git a/Engine/Source/Runtime/Engine/Private/UnrealClient.cpp b/Engine/Source/Runtime/Engine/Private/UnrealClient.cpp
--- a/Engine/Source/Runtime/Engine/Private/UnrealClient.cpp
+++ b/Engine/Source/Runtime/Engine/Private/UnrealClient.cpp
@@ -1229,7 +1229,7 @@ int32 FStatUnitData::DrawStat(FViewport* InViewport, FCanvas* InCanvas, int32 In
 				GraphVerticalPixelPerValue = GraphVerticalPixelsPerMS;
 				StatColor = FLinearColor(1.0f, 0.1f, 1.0f);		// Green
 			}
-			else if (StatIndex == EGS_DynRes)
+			else if (StatIndex == EGS_DynRes) //-V547
 			{
 				const DynamicRenderScaling::FBudget& Budget = GDynamicPrimaryResolutionFraction;
 
diff --git a/Engine/Source/Runtime/GeometryCore/Private/DynamicMesh/DynamicMeshAttributeSet.cpp b/Engine/Source/Runtime/GeometryCore/Private/DynamicMesh/DynamicMeshAttributeSet.cpp
--- a/Engine/Source/Runtime/GeometryCore/Private/DynamicMesh/DynamicMeshAttributeSet.cpp
+++ b/Engine/Source/Runtime/GeometryCore/Private/DynamicMesh/DynamicMeshAttributeSet.cpp
@@ -487,7 +487,7 @@ void FDynamicMeshAttributeSet::EnableMatchingAttributes(const FDynamicMeshAttrib
 	for (int32 k = bClearExisting ? 0 : ExistingPolygroupLayers; k < NumPolygroupLayers(); k++)
 	{
 		PolygroupLayers[k].Initialize((int32)0);
-		if (PolygroupLayers[k].GetName() == NAME_None && k < ToMatch.NumPolygroupLayers())
+		if (k < ToMatch.NumPolygroupLayers() && PolygroupLayers[k].GetName() == NAME_None)
 		{
 			PolygroupLayers[k].SetName( ToMatch.GetPolygroupLayer(k)->GetName() );
 		}
@@ -500,7 +500,7 @@ void FDynamicMeshAttributeSet::EnableMatchingAttributes(const FDynamicMeshAttrib
 	for (int32 k = bClearExisting ? 0 : ExistingWeightLayers; k < NumWeightLayers(); k++)
 	{
 		WeightLayers[k].Initialize(0.0f);
-		if (WeightLayers[k].GetName() == NAME_None && k < ToMatch.NumWeightLayers())
+		if (k < ToMatch.NumWeightLayers() && WeightLayers[k].GetName() == NAME_None)
 		{
 			WeightLayers[k].SetName( ToMatch.GetWeightLayer(k)->GetName() );
 		}
diff --git a/Engine/Source/Runtime/GeometryCore/Private/DynamicMeshEditor.cpp b/Engine/Source/Runtime/GeometryCore/Private/DynamicMeshEditor.cpp
--- a/Engine/Source/Runtime/GeometryCore/Private/DynamicMeshEditor.cpp
+++ b/Engine/Source/Runtime/GeometryCore/Private/DynamicMeshEditor.cpp
@@ -1090,7 +1090,7 @@ bool FDynamicMeshEditor::ReinsertSubmesh(const FDynamicSubmesh3& Region, FOption
 		{
 			int sub_v = sub_t[j];
 			int new_v = -1;
-			if (done_v[sub_v] == false)
+			if (done_v[sub_v] == false) //-V562
 			{
 				// first check if this is a boundary vtx on submesh and maps to a bdry vtx on base mesh
 				if (Sub.IsBoundaryVertex(sub_v))
diff --git a/Engine/Source/Runtime/GeometryCore/Private/MeshRegionBoundaryLoops.cpp b/Engine/Source/Runtime/GeometryCore/Private/MeshRegionBoundaryLoops.cpp
--- a/Engine/Source/Runtime/GeometryCore/Private/MeshRegionBoundaryLoops.cpp
+++ b/Engine/Source/Runtime/GeometryCore/Private/MeshRegionBoundaryLoops.cpp
@@ -100,7 +100,7 @@ bool FMeshRegionBoundaryLoops::Compute()
 	// process all edges of mesh
 	for (int eid : edges_roi)
 	{
-		if (used_edge[eid] == true)
+		if (used_edge[eid] == true) //-V562
 		{
 			continue;
 		}
@@ -212,7 +212,7 @@ bool FMeshRegionBoundaryLoops::Compute()
 			else
 			{
 				// push onto accumulated list
-				check(used_edge[eNext] == false);
+				check(used_edge[eNext] == false); //-V562
 				loop_edges.Add(eNext);
 				eCur = eNext;
 				used_edge.Add(eCur);
diff --git a/Engine/Source/Runtime/MovieScene/Private/Sections/MovieSceneSpawnSection.cpp b/Engine/Source/Runtime/MovieScene/Private/Sections/MovieSceneSpawnSection.cpp
--- a/Engine/Source/Runtime/MovieScene/Private/Sections/MovieSceneSpawnSection.cpp
+++ b/Engine/Source/Runtime/MovieScene/Private/Sections/MovieSceneSpawnSection.cpp
@@ -104,7 +104,7 @@ bool UMovieSceneSpawnSection::PopulateEvaluationFieldImpl(const TRange<FFrameNum
 			break;
 		}
 
-		if (bIsSpawned != Values[Index])
+		if (bIsSpawned != Values[Index]) //-V562
 		{
 			if (bIsSpawned)
 			{
diff --git a/Engine/Source/Runtime/MovieScene/Private/Tracks/MovieSceneSpawnTrack.cpp b/Engine/Source/Runtime/MovieScene/Private/Tracks/MovieSceneSpawnTrack.cpp
--- a/Engine/Source/Runtime/MovieScene/Private/Tracks/MovieSceneSpawnTrack.cpp
+++ b/Engine/Source/Runtime/MovieScene/Private/Tracks/MovieSceneSpawnTrack.cpp
@@ -105,7 +105,7 @@ void UMovieSceneSpawnTrack::PopulateSpawnedRangeMask(const TRange<FFrameNumber>&
 				break;
 			}
 
-			if (bIsSpawned != Values[Index])
+			if (bIsSpawned != Values[Index]) //-V562
 			{
 				if (bIsSpawned)
 				{
diff --git a/Engine/Source/Runtime/Navmesh/Private/DetourTileCache/DetourTileCacheBuilder.cpp b/Engine/Source/Runtime/Navmesh/Private/DetourTileCache/DetourTileCacheBuilder.cpp
--- a/Engine/Source/Runtime/Navmesh/Private/DetourTileCache/DetourTileCacheBuilder.cpp
+++ b/Engine/Source/Runtime/Navmesh/Private/DetourTileCache/DetourTileCacheBuilder.cpp
@@ -782,6 +782,7 @@ static void getClosestIndices(const unsigned short* vertsa, const int nvertsa,
 
 static bool mergeContours(dtTileCacheAlloc* alloc, dtTileCacheContour& ca, dtTileCacheContour& cb, int ia, int ib)
 {
+	check(ca.nverts > 0 && cb.nverts > 0);
 	const int maxVerts = ca.nverts + cb.nverts + 2;
 	unsigned short* verts = (unsigned short*)alloc->alloc(sizeof(unsigned short)*maxVerts*4);
 	if (!verts)
diff --git a/Engine/Source/Runtime/Navmesh/Private/Recast/RecastContour.cpp b/Engine/Source/Runtime/Navmesh/Private/Recast/RecastContour.cpp
--- a/Engine/Source/Runtime/Navmesh/Private/Recast/RecastContour.cpp
+++ b/Engine/Source/Runtime/Navmesh/Private/Recast/RecastContour.cpp
@@ -548,6 +548,7 @@ static void getClosestIndices(const int* vertsa, const int nvertsa,
 
 static bool mergeContours(rcContour& ca, rcContour& cb, int ia, int ib)
 {
+	check(ca.nverts > 0 && cb.nverts > 0);
 	const int maxVerts = ca.nverts + cb.nverts + 2;
 	int* verts = (int*)rcAlloc(sizeof(int)*maxVerts*4, RC_ALLOC_PERM);
 	if (!verts)
diff --git a/Engine/Source/Runtime/TypedElementFramework/Private/DataStorage/MapKey.cpp b/Engine/Source/Runtime/TypedElementFramework/Private/DataStorage/MapKey.cpp
--- a/Engine/Source/Runtime/TypedElementFramework/Private/DataStorage/MapKey.cpp
+++ b/Engine/Source/Runtime/TypedElementFramework/Private/DataStorage/MapKey.cpp
@@ -139,7 +139,7 @@ namespace UE::Editor::DataStorage
 
 		bool operator()(const TUniquePtr<FSoftObjectPath>& Other)
 		{
-			if (const TUniquePtr<FSoftObjectPath>* Local = Key.TryGet<TUniquePtr<FSoftObjectPath>>())
+			if (const TUniquePtr<FSoftObjectPath>* Local = Key.TryGet<TUniquePtr<FSoftObjectPath>>()) //-V593
 			{
 				if (Local->IsValid() && Other.IsValid())
 				{
diff --git a/Engine/Source/Runtime/UniversalObjectLocator/Public/UniversalObjectLocatorFragment.h b/Engine/Source/Runtime/UniversalObjectLocator/Public/UniversalObjectLocatorFragment.h
--- a/Engine/Source/Runtime/UniversalObjectLocator/Public/UniversalObjectLocatorFragment.h
+++ b/Engine/Source/Runtime/UniversalObjectLocator/Public/UniversalObjectLocatorFragment.h
@@ -39,10 +39,10 @@ struct alignas(8) FUniversalObjectLocatorFragment
 	using FParseStringResult = UE::UniversalObjectLocator::FParseStringResult;
 	using FParseStringParams = UE::UniversalObjectLocator::FParseStringParams;
 
-	static constexpr FAsciiSet ValidFragmentTypeCharacters = "abcdefghijklmnopqrstuvwxyzABCDEFGHIJKLMNOPQRSTUVWXYZ0123456789_-";
+	static inline constexpr FAsciiSet ValidFragmentTypeCharacters = "abcdefghijklmnopqrstuvwxyzABCDEFGHIJKLMNOPQRSTUVWXYZ0123456789_-";
 	// Valid fragment delimiters are those that are allowable by RFC3986 for the query part (unreserved / pct-encoded / sub-delims / ":" / "@" / "/" / "?") excluding '&' which we use to separate fragments within a path
-	static constexpr FAsciiSet ValidFragmentDelimiters = "%!$'()*+,;=/?:@.~";
-	static constexpr FAsciiSet ValidFragmentPayloadCharacters = ValidFragmentTypeCharacters | ValidFragmentDelimiters;
+	static inline constexpr FAsciiSet ValidFragmentDelimiters = "%!$'()*+,;=/?:@.~";
+	static inline constexpr FAsciiSet ValidFragmentPayloadCharacters = ValidFragmentTypeCharacters | ValidFragmentDelimiters;
 
 	/** Make our inline data buffer larger in-editor to support editor-only data without allocation */
 #if WITH_EDITORONLY_DATA
