Downgrade UOL invalid character assertion to a non-fatal ensure

This error, although serious, should not bring down the process with a fatal assert.

diff --git a/Engine/Source/Runtime/UniversalObjectLocator/Private/UniversalObjectLocatorFragmentType.cpp b/Engine/Source/Runtime/UniversalObjectLocator/Private/UniversalObjectLocatorFragmentType.cpp
--- a/Engine/Source/Runtime/UniversalObjectLocator/Private/UniversalObjectLocatorFragmentType.cpp
+++ b/Engine/Source/Runtime/UniversalObjectLocator/Private/UniversalObjectLocatorFragmentType.cpp
@@ -52,7 +52,7 @@ void FFragmentType::ToString(const void* Payload, FStringBuilderBase& OutStringB
 	};
 
 	FStringView StringRepresentation(OutStringBuilder.GetData() + StartPos, OutStringBuilder.Len());
-	checkf(!FAsciiSet::HasAny(StringRepresentation, InvalidChars),
+	ensureMsgf(!FAsciiSet::HasAny(StringRepresentation, InvalidChars),
 		TEXT("F%s::ToString resulted in an invalid character usage:\n\t%s\n\t%s"),
 		*PayloadType->GetName(),
 		*FString(StringRepresentation.Len(), StringRepresentation.GetData()),
