Fixed potential for buffer overrun in error reporting string

diff --git a/Engine/Source/Runtime/UniversalObjectLocator/Private/UniversalObjectLocatorFragmentType.cpp b/Engine/Source/Runtime/UniversalObjectLocator/Private/UniversalObjectLocatorFragmentType.cpp
--- a/Engine/Source/Runtime/UniversalObjectLocator/Private/UniversalObjectLocatorFragmentType.cpp
+++ b/Engine/Source/Runtime/UniversalObjectLocator/Private/UniversalObjectLocatorFragmentType.cpp
@@ -51,7 +51,7 @@ void FFragmentType::ToString(const void* Payload, FStringBuilderBase& OutStringB
 		return String;
 	};
 
-	FStringView StringRepresentation(OutStringBuilder.GetData() + StartPos, OutStringBuilder.Len());
+	FStringView StringRepresentation = FStringView(OutStringBuilder).RightChop(StartPos);
 	ensureMsgf(!FAsciiSet::HasAny(StringRepresentation, InvalidChars),
 		TEXT("F%s::ToString resulted in an invalid character usage:\n\t%s\n\t%s"),
 		*PayloadType->GetName(),
