Sequencer: Encode/decode fragment type string since invalid object name characters aren't enforced.

#Synced-CL 48303336

diff --git a/Engine/Source/Runtime/UniversalObjectLocator/Private/UniversalObjectLocatorFragmentType.cpp b/Engine/Source/Runtime/UniversalObjectLocator/Private/UniversalObjectLocatorFragmentType.cpp
--- a/Engine/Source/Runtime/UniversalObjectLocator/Private/UniversalObjectLocatorFragmentType.cpp
+++ b/Engine/Source/Runtime/UniversalObjectLocator/Private/UniversalObjectLocatorFragmentType.cpp
@@ -6,6 +6,8 @@
 #include "UniversalObjectLocatorInitializeParams.h"
 #include "UniversalObjectLocatorInitializeResult.h"
 
+#include "GenericPlatform/GenericPlatformHttp.h"
+
 namespace UE::UniversalObjectLocator
 {
 
@@ -26,10 +28,14 @@ void FFragmentType::ToString(const void* Payload, FStringBuilderBase& OutStringB
 		OutStringBuilder.RemoveSuffix(1);
 	}
 
-	const int32 StartPos = OutStringBuilder.Len();
-
 	(*InstanceBindings.ToString)(Payload, OutStringBuilder);
 
+	const FStringView StringToEncode = OutStringBuilder.ToString();
+	FString EncodedString = FGenericPlatformHttp::UrlEncode(StringToEncode);
+	OutStringBuilder = EncodedString;
+
+	const int32 StartPos = OutStringBuilder.Len();
+
 #if DO_CHECK
 
 	static constexpr FAsciiSet InvalidChars = ~FUniversalObjectLocatorFragment::ValidFragmentPayloadCharacters;
@@ -61,7 +67,8 @@ void FFragmentType::ToString(const void* Payload, FStringBuilderBase& OutStringB
 }
 FParseStringResult FFragmentType::TryParseString(void* Payload, FStringView InString, const FParseStringParams& Params) const
 {
-	return (*InstanceBindings.TryParseString)(Payload, InString, Params);
+	FString DecodedString = FGenericPlatformHttp::UrlDecode(InString);
+	return (*InstanceBindings.TryParseString)(Payload, DecodedString, Params);
 }
 
 uint32 FFragmentType::ComputePriority(const UObject* Object, const UObject* Context) const
diff --git a/Engine/Source/Runtime/UniversalObjectLocator/UniversalObjectLocator.Build.cs b/Engine/Source/Runtime/UniversalObjectLocator/UniversalObjectLocator.Build.cs
--- a/Engine/Source/Runtime/UniversalObjectLocator/UniversalObjectLocator.Build.cs
+++ b/Engine/Source/Runtime/UniversalObjectLocator/UniversalObjectLocator.Build.cs
@@ -12,6 +12,7 @@ public class UniversalObjectLocator : ModuleRules
 			new string[] {
 				"Core",
 				"CoreUObject",
+				"HTTP",
 			}
 		);
 	}
